World Conditions
- General purpose gameplay conditions
- Supports building expression from conditions
- Allows to cache condition results and only eval when things change

diff --git a/Engine/Plugins/Runtime/WorldConditions/Source/WorldConditions/Private/WorldConditionBase.cpp b/Engine/Plugins/Runtime/WorldConditions/Source/WorldConditions/Private/WorldConditionBase.cpp
new file mode 100644
--- /dev/null
+++ b/Engine/Plugins/Runtime/WorldConditions/Source/WorldConditions/Private/WorldConditionBase.cpp
@@ -0,0 +1,40 @@
+﻿// Copyright Epic Games, Inc. All Rights Reserved.
+
+#include "WorldConditionBase.h"
+#include "WorldConditionQuery.h"
+
+FWorldConditionBase::~FWorldConditionBase()
+{
+	// Empty
+}
+
+FText FWorldConditionBase::GetDescription() const
+{
+	return StaticStruct()->GetDisplayNameText();
+}
+
+bool FWorldConditionBase::Initialize(const UWorldConditionSchema& Schema)
+{
+	return true;
+}
+
+bool FWorldConditionBase::Activate(const FWorldConditionContext& Context) const
+{
+	return true;
+}
+
+EWorldConditionResult FWorldConditionBase::IsTrue(const FWorldConditionContext& Context) const
+{
+	return EWorldConditionResult::IsTrue;
+}
+
+void FWorldConditionBase::Deactivate(const FWorldConditionContext& Context) const
+{
+	// Empty
+}
+
+void FWorldConditionBase::InvalidateResult(FWorldConditionQueryState& QueryState) const
+{
+	QueryState.CachedResult = EWorldConditionResult::Invalid;
+	QueryState.GetItem(ConditionIndex).CachedResult = EWorldConditionResult::Invalid; 
+}
diff --git a/Engine/Plugins/Runtime/WorldConditions/Source/WorldConditions/Private/WorldConditionContext.cpp b/Engine/Plugins/Runtime/WorldConditions/Source/WorldConditions/Private/WorldConditionContext.cpp
new file mode 100644
--- /dev/null
+++ b/Engine/Plugins/Runtime/WorldConditions/Source/WorldConditions/Private/WorldConditionContext.cpp
@@ -0,0 +1,101 @@
+﻿// Copyright Epic Games, Inc. All Rights Reserved.
+
+#include "WorldConditionContext.h"
+
+
+bool FWorldConditionContext::Activate() const
+{
+	check(QueryState.IsValid());
+
+	for (int32 Index = 0; Index < QueryDefinition.Conditions.Num(); Index++)
+	{
+		const FWorldConditionBase& Condition = QueryDefinition.Conditions[Index].Get<FWorldConditionBase>();
+		FWorldConditionItem& Item = QueryState.GetItem(Index);
+		Item.Operator = Condition.Operator;
+		Item.NextExpressionDepth = Condition.NextExpressionDepth;
+	}
+
+	bool bSuccess = true;
+	for (int32 Index = 0; Index < QueryDefinition.Conditions.Num(); Index++)
+	{
+		const FWorldConditionBase& Condition = QueryDefinition.Conditions[Index].Get<FWorldConditionBase>();
+		bSuccess &= Condition.Activate(*this);
+	}
+
+	if (!bSuccess)
+	{
+		Deactivate();
+		QueryState.Free(QueryDefinition);
+	}
+	
+	return bSuccess;
+}
+
+bool FWorldConditionContext::IsTrue() const
+{
+	if (!QueryState.IsValid())
+	{
+		return false;
+	}
+
+	if (QueryState.CachedResult != EWorldConditionResult::Invalid)
+	{
+		return QueryState.CachedResult == EWorldConditionResult::IsTrue;
+	}
+	
+	static_assert(UE::WorldCondition::MaxExpressionDepth == 4);
+	EWorldConditionResult Results[UE::WorldCondition::MaxExpressionDepth + 1] = { EWorldConditionResult::Invalid, EWorldConditionResult::Invalid, EWorldConditionResult::Invalid, EWorldConditionResult::Invalid, EWorldConditionResult::Invalid };
+	EWorldConditionOperator Operators[UE::WorldCondition::MaxExpressionDepth + 1] = { EWorldConditionOperator::Copy, EWorldConditionOperator::Copy, EWorldConditionOperator::Copy, EWorldConditionOperator::Copy, EWorldConditionOperator::Copy };
+	int32 Depth = 0;
+
+	for (int32 Index = 0; Index < QueryState.GetNumConditions(); Index++)
+	{
+		FWorldConditionItem& Item = QueryState.GetItem(Index);
+		const int32 NextExpressionDepth = Item.NextExpressionDepth;
+
+		Operators[Depth] = Item.Operator;
+
+		Depth = FMath::Max(Depth, NextExpressionDepth);
+		
+		EWorldConditionResult CurrResult = Item.CachedResult;
+		if (CurrResult == EWorldConditionResult::Invalid)
+		{
+			check(QueryDefinition.Conditions.Num() == QueryState.GetNumConditions());
+			const FWorldConditionBase& Condition = QueryDefinition.Conditions[Index].Get<FWorldConditionBase>();
+			CurrResult = Condition.IsTrue(*this);
+			if (Condition.bCanCacheResult)
+			{
+				Item.CachedResult = CurrResult;
+			}
+		}
+
+		Depth++;
+		Results[Depth] = CurrResult;
+
+		while (Depth > NextExpressionDepth)
+		{
+			Depth--;
+			Results[Depth] = UE::WorldCondition::MergeResults(Operators[Depth], Results[Depth], Results[Depth + 1]);
+			Operators[Depth] = EWorldConditionOperator::Copy;
+		}
+	}
+
+	QueryState.CachedResult = Results[0]; 
+	
+	return QueryState.CachedResult == EWorldConditionResult::IsTrue;
+
+}
+
+void FWorldConditionContext::Deactivate() const
+{
+	check(QueryState.IsValid());
+	
+	for (int32 Index = 0; Index < QueryDefinition.Conditions.Num(); Index++)
+	{
+		const FWorldConditionBase& ConditionDef = QueryDefinition.Conditions[Index].Get<FWorldConditionBase>();
+		ConditionDef.Deactivate(*this);
+	}
+
+	QueryState.Free(QueryDefinition);
+
+}
diff --git a/Engine/Plugins/Runtime/WorldConditions/Source/WorldConditions/Private/WorldConditionQuery.cpp b/Engine/Plugins/Runtime/WorldConditions/Source/WorldConditions/Private/WorldConditionQuery.cpp
new file mode 100644
--- /dev/null
+++ b/Engine/Plugins/Runtime/WorldConditions/Source/WorldConditions/Private/WorldConditionQuery.cpp
@@ -0,0 +1,330 @@
+﻿// Copyright Epic Games, Inc. All Rights Reserved.
+
+#include "WorldConditionQuery.h"
+#include "WorldConditionContext.h"
+
+FWorldConditionQueryState::~FWorldConditionQueryState()
+{
+	if (Memory)
+	{
+		UE_LOG(LogWorldCondition, Error, TEXT("Expected World Condition state %p to have been freed, might leak memory."), Memory);
+		FMemory::Free(Memory);
+		Memory = nullptr;
+	}
+}
+	
+bool FWorldConditionQueryState::Initialize(const UObject& Owner, const FWorldConditionQueryDefinition& QueryDefinition)
+{
+	check(Memory == nullptr);
+
+	if (!QueryDefinition.IsValid())
+	{
+		UE_LOG(LogWorldCondition, Error, TEXT("%s: Trying to intialize world condition with invalid defintion."), *GetFullNameSafe(&Owner));
+		return false;
+	}
+
+	if (IsValid())
+	{
+		Free(QueryDefinition);
+	}
+	
+	NumConditions = IntCastChecked<uint8>(QueryDefinition.Conditions.Num());
+
+	int32 MinAlignment = 8;
+	int32 Offset = 0;
+	
+	// Reserve space for condition items.
+	Offset += sizeof(FWorldConditionItem) * static_cast<int32>(NumConditions); 
+
+	bHasPerConditionState = false;
+	
+	// Reserve space for all runtime data.
+	for (int32 Index = 0; Index < static_cast<int32>(NumConditions); Index++)
+	{
+		FWorldConditionBase& Condition = QueryDefinition.Conditions[Index].GetMutable<FWorldConditionBase>();
+		if (const UStruct* StateStruct = Condition.GetRuntimeStateType())
+		{
+			int32 StructMinAlignment = 0;
+			int32 StructSize = 0;
+
+			if (const UScriptStruct* ScriptStruct = Cast<const UScriptStruct>(StateStruct))
+			{
+				StructMinAlignment = ScriptStruct->GetMinAlignment();
+				StructSize = ScriptStruct->GetStructureSize();
+				Condition.bIsStateObject = false;
+			}
+			else if (const UClass* Class = Cast<const UClass>(StateStruct))
+			{
+				StructMinAlignment = FWorldConditionStateObject::StaticStruct()->GetMinAlignment();
+				StructSize = FWorldConditionStateObject::StaticStruct()->GetStructureSize();
+				Condition.bIsStateObject = true;
+			}
+			
+			check(StructMinAlignment > 0 && StructSize > 0);
+
+			Offset = Align(Offset, StructMinAlignment);
+			Condition.StateDataOffset = IntCastChecked<uint16>(Offset);
+
+			Offset += StructSize;
+			MinAlignment = FMath::Max(MinAlignment, StructMinAlignment);
+
+			bHasPerConditionState = true;
+		}
+		else
+		{
+			Condition.StateDataOffset = 0;
+		}
+	}
+
+	const int32 TotalSize = Offset;
+	Memory = (uint8*)FMemory::Malloc(TotalSize, MinAlignment);
+
+	// Initialize items
+	for (int32 Index = 0; Index < static_cast<int32>(NumConditions); Index++)
+	{
+		new (Memory + sizeof(FWorldConditionItem) * Index) FWorldConditionItem();
+	}
+
+	// Initialize state
+	for (int32 Index = 0; Index < static_cast<int32>(NumConditions); Index++)
+	{
+		FWorldConditionBase& Condition = QueryDefinition.Conditions[Index].GetMutable<FWorldConditionBase>();
+		if (Condition.StateDataOffset == 0)
+		{
+			continue;
+		}
+		uint8* StateMemory = Memory + Condition.StateDataOffset;
+		if (Condition.bIsStateObject)
+		{
+			new (StateMemory) FWorldConditionStateObject();
+			FWorldConditionStateObject& StateObject = *reinterpret_cast<FWorldConditionStateObject*>(StateMemory);
+			const UClass* StateClass = Cast<UClass>(Condition.GetRuntimeStateType());
+			StateObject.Object = NewObject<UObject>(const_cast<UObject*>(&Owner), StateClass);  
+		}
+		else
+		{
+			const UScriptStruct* StateScriptStruct = Cast<UScriptStruct>(Condition.GetRuntimeStateType());
+			StateScriptStruct->InitializeStruct(StateMemory);
+		}
+	}
+
+	return true;
+}
+
+void FWorldConditionQueryState::Free(const FWorldConditionQueryDefinition& QueryDefinition)
+{
+	check(Memory != nullptr);
+	
+	// Items don't need destructing.
+
+	// Destroy state
+	for (int32 Index = 0; Index < static_cast<int32>(NumConditions); Index++)
+	{
+		FWorldConditionBase& Condition = QueryDefinition.Conditions[Index].GetMutable<FWorldConditionBase>();
+		if (Condition.StateDataOffset == 0)
+		{
+			continue;
+		}
+		uint8* StateMemory = Memory + Condition.StateDataOffset;
+		if (Condition.bIsStateObject)
+		{
+			FWorldConditionStateObject& StateObject = *reinterpret_cast<FWorldConditionStateObject*>(Memory + Condition.StateDataOffset);
+			StateObject.~FWorldConditionStateObject();
+		}
+		else
+		{
+			const UScriptStruct* StateScriptStruct = Cast<UScriptStruct>(Condition.GetRuntimeStateType());
+			StateScriptStruct->DestroyStruct(StateMemory);
+		}
+	}
+
+	FMemory::Free(Memory);
+	Memory = nullptr;
+	NumConditions = 0;
+	CachedResult = EWorldConditionResult::Invalid;
+	bHasPerConditionState = false;
+}
+
+void FWorldConditionQueryState::AddReferencedObjects(const FWorldConditionQueryDefinition& QueryDefinition, class FReferenceCollector& Collector) const
+{
+	if (Memory == nullptr)
+	{
+		return;
+	}
+	
+	check(NumConditions == QueryDefinition.Conditions.Num());
+	
+	for (int32 Index = 0; Index < QueryDefinition.Conditions.Num(); Index++)
+	{
+		FWorldConditionBase& Condition = QueryDefinition.Conditions[Index].GetMutable<FWorldConditionBase>();
+		if (Condition.StateDataOffset == 0)
+		{
+			continue;
+		}
+		uint8* StateMemory = Memory + Condition.StateDataOffset;
+		if (Condition.bIsStateObject)
+		{
+			FWorldConditionStateObject& StateObject = *reinterpret_cast<FWorldConditionStateObject*>(Memory + Condition.StateDataOffset);
+			Collector.AddReferencedObject(StateObject.Object);
+		}
+		else
+		{
+			const UScriptStruct* StateScriptStruct = Cast<UScriptStruct>(Condition.GetRuntimeStateType());
+			Collector.AddReferencedObjects(StateScriptStruct, StateMemory);
+		}
+	}
+	
+}
+
+
+
+//
+// FWorldConditionQueryDefinition
+//
+
+bool FWorldConditionQueryDefinition::IsValid() const
+{
+	return SchemaClass && Conditions.Num() > 0;
+}
+
+bool FWorldConditionQueryDefinition::Initialize()
+{
+	bool bResult = true;
+#if WITH_EDITORONLY_DATA
+	Conditions.Reset();
+
+	if (SchemaClass == nullptr)
+	{
+		return false;
+	}
+
+	const UWorldConditionSchema* Schema = SchemaClass.GetDefaultObject();
+	
+	// Append only valid condition.
+	TArray<FConstStructView> ValidConditions;
+	ValidConditions.Reserve(EditableConditions.Num());
+	for (const FWorldConditionEditable& EditableCondition : EditableConditions)
+	{
+		if (EditableCondition.Condition.IsValid())
+		{
+			if (Schema->IsStructAllowed(EditableCondition.Condition.GetScriptStruct()))
+			{
+				FWorldConditionBase& Condition = EditableCondition.Condition.GetMutable<FWorldConditionBase>();
+				// Store expression depth temporarily into NextExpressionDepth, it will be update below.
+				Condition.NextExpressionDepth = EditableCondition.ExpressionDepth;
+				Condition.Operator = EditableCondition.Operator;
+
+				ValidConditions.Add(EditableCondition.Condition);
+			}
+			else
+			{
+				UE_LOG(LogWorldCondition, Warning, TEXT("World condition contains condition of type %s that is not allowed by schema %s."),
+					*GetNameSafe(EditableCondition.Condition.GetScriptStruct()), *GetNameSafe(Schema));
+			}
+		}
+	}
+	if (ValidConditions.IsEmpty())
+	{
+		UE_LOG(LogWorldCondition, Error, TEXT("World condition requires at least one valid condition."));
+		return false;
+	}
+	
+	Conditions.Append(ValidConditions);
+
+	if (Conditions.Num() > 0)
+	{
+		FWorldConditionBase& Condition = Conditions[0].GetMutable<FWorldConditionBase>();
+		Condition.Operator = EWorldConditionOperator::Copy;
+	}
+
+	for (int32 Index = 0; Index < Conditions.Num(); Index++)
+	{
+		uint8 NextExpressionDepth = 0;
+		if ((Index + 1) < Conditions.Num())
+		{
+			const FWorldConditionBase& NextCondition = Conditions[Index + 1].GetMutable<FWorldConditionBase>();
+			NextExpressionDepth = NextCondition.NextExpressionDepth;
+		}
+		
+		FWorldConditionBase& Condition = Conditions[Index].GetMutable<FWorldConditionBase>();
+		Condition.NextExpressionDepth = NextExpressionDepth;
+
+		Condition.ConditionIndex = Index;
+	}
+
+	for (int32 Index = 0; Index < Conditions.Num(); Index++)
+	{
+		FWorldConditionBase& Condition = Conditions[Index].GetMutable<FWorldConditionBase>();
+		bResult &= Condition.Initialize(*Schema);
+	}
+#endif
+	
+	return bResult;
+}
+
+
+//
+// FWorldConditionQuery
+//
+
+FWorldConditionQuery::~FWorldConditionQuery()
+{
+	if (IsActive())
+	{
+		UE_LOG(LogWorldCondition, Error, TEXT("Active World Condition %p is still active on destructor, calling Deactivate() without context data, might leak memory."), this);
+		check(QueryDefinition.SchemaClass);
+		const FWorldConditionContextData ContextData(*QueryDefinition.SchemaClass.GetDefaultObject());
+		Deactivate(ContextData);
+	}
+}
+
+bool FWorldConditionQuery::DebugInitialize(const TSubclassOf<UWorldConditionSchema> InSchemaClass, const TConstArrayView<FWorldConditionEditable> InConditions)
+{
+	if (IsActive())
+	{
+		return false;
+	}
+
+	QueryDefinition.SchemaClass = InSchemaClass;
+	QueryDefinition.EditableConditions = InConditions;
+	return QueryDefinition.Initialize();
+}
+
+bool FWorldConditionQuery::Activate(const UObject& InOwner, const FWorldConditionContextData& ContextData)
+{
+	Owner = &InOwner;
+
+	if (!QueryState.Initialize(*Owner, QueryDefinition))
+	{
+		return false;
+	}
+	check(QueryState.GetNumConditions() == QueryDefinition.Conditions.Num());
+
+	const FWorldConditionContext Context(*Owner, QueryDefinition, QueryState, ContextData);
+	return Context.Activate();
+}
+
+bool FWorldConditionQuery::IsTrue(const FWorldConditionContextData& ContextData) const
+{
+	const FWorldConditionContext Context(*Owner, QueryDefinition, QueryState, ContextData);
+	return Context.IsTrue();
+}
+
+void FWorldConditionQuery::Deactivate(const FWorldConditionContextData& ContextData) const
+{
+	const FWorldConditionContext Context(*Owner, QueryDefinition, QueryState, ContextData);
+	return Context.Deactivate();
+}
+
+bool FWorldConditionQuery::IsActive() const
+{
+	return QueryState.IsValid();
+}
+
+void FWorldConditionQuery::AddStructReferencedObjects(class FReferenceCollector& Collector) const
+{
+	if (QueryState.IsValid())
+	{
+		QueryState.AddReferencedObjects(QueryDefinition, Collector);
+	}
+}
diff --git a/Engine/Plugins/Runtime/WorldConditions/Source/WorldConditions/Private/WorldConditionSchema.cpp b/Engine/Plugins/Runtime/WorldConditions/Source/WorldConditions/Private/WorldConditionSchema.cpp
new file mode 100644
--- /dev/null
+++ b/Engine/Plugins/Runtime/WorldConditions/Source/WorldConditions/Private/WorldConditionSchema.cpp
@@ -0,0 +1,52 @@
+﻿// Copyright Epic Games, Inc. All Rights Reserved.
+
+#include "WorldConditionSchema.h"
+
+
+TConstArrayView<FWorldConditionContextDataDesc> UWorldConditionSchema::GetContextDataDescs() const
+{
+	return ContextDataDescs;
+}
+
+const FWorldConditionContextDataDesc* UWorldConditionSchema::GetContextDataDescByRef(const FWorldConditionContextDataRef& Ref) const
+{
+	return Ref.IsValid() ? &ContextDataDescs[Ref.Index] : nullptr;
+}
+
+const FWorldConditionContextDataDesc& UWorldConditionSchema::GetContextDataDescByIndex(const int32 Index) const
+{
+	return ContextDataDescs[Index];
+}
+
+EWorldConditionContextDataType UWorldConditionSchema::GetContextDataTypeByRef(const FWorldConditionContextDataRef& Ref) const
+{
+	return Ref.IsValid() ? ContextDataDescs[Ref.Index].Type : EWorldConditionContextDataType::Persistent;
+}
+
+EWorldConditionContextDataType UWorldConditionSchema::GetContextDataTypeByIndex(const int32 Index) const
+{
+	return ContextDataDescs[Index].Type;
+}
+
+int32 UWorldConditionSchema::GetContextDataIndexByName(const FName DataName, const UStruct* Struct) const
+{
+	return ContextDataDescs.IndexOfByPredicate([DataName, Struct](const FWorldConditionContextDataDesc& Desc)
+	{
+		return Desc.Name == DataName && Desc.Struct == Struct;
+	});
+}
+
+const FWorldConditionContextDataDesc* UWorldConditionSchema::GetContextDataDescByName(const FName DataName, const UStruct* Struct) const
+{
+	return ContextDataDescs.FindByPredicate([DataName, Struct](const FWorldConditionContextDataDesc& Desc)
+	{
+		return Desc.Name == DataName && Desc.Struct == Struct;
+	});
+}
+
+FWorldConditionContextDataRef UWorldConditionSchema::AddContextDataDesc(const FName InName, const UStruct* InStruct, const EWorldConditionContextDataType InType)
+{
+	const int32 Index = ContextDataDescs.Num();
+	ContextDataDescs.Emplace(InName, InStruct, InType);
+	return FWorldConditionContextDataRef(InName, Index);
+}
diff --git a/Engine/Plugins/Runtime/WorldConditions/Source/WorldConditions/Private/WorldConditionTypes.cpp b/Engine/Plugins/Runtime/WorldConditions/Source/WorldConditions/Private/WorldConditionTypes.cpp
new file mode 100644
--- /dev/null
+++ b/Engine/Plugins/Runtime/WorldConditions/Source/WorldConditions/Private/WorldConditionTypes.cpp
@@ -0,0 +1,6 @@
+﻿// Copyright Epic Games, Inc. All Rights Reserved.
+
+#include "WorldConditionTypes.h"
+#include "WorldConditionSchema.h"
+
+DEFINE_LOG_CATEGORY(LogWorldCondition);
diff --git a/Engine/Plugins/Runtime/WorldConditions/Source/WorldConditions/Private/WorldConditionsModule.cpp b/Engine/Plugins/Runtime/WorldConditions/Source/WorldConditions/Private/WorldConditionsModule.cpp
new file mode 100644
--- /dev/null
+++ b/Engine/Plugins/Runtime/WorldConditions/Source/WorldConditions/Private/WorldConditionsModule.cpp
@@ -0,0 +1,24 @@
+// Copyright Epic Games, Inc. All Rights Reserved.
+
+#include "WorldConditionsModule.h"
+#include "Modules/ModuleManager.h"
+
+#define LOCTEXT_NAMESPACE "WorldConditions"
+
+class FWorldConditionsModule : public IWorldConditionsModule
+{
+	virtual void StartupModule() override;
+	virtual void ShutdownModule() override;
+};
+
+IMPLEMENT_MODULE(FWorldConditionsModule, WorldConditionsModule)
+
+void FWorldConditionsModule::StartupModule()
+{
+}
+
+void FWorldConditionsModule::ShutdownModule()
+{
+}
+
+#undef LOCTEXT_NAMESPACE
diff --git a/Engine/Plugins/Runtime/WorldConditions/Source/WorldConditions/Public/WorldConditionBase.h b/Engine/Plugins/Runtime/WorldConditions/Source/WorldConditions/Public/WorldConditionBase.h
new file mode 100644
--- /dev/null
+++ b/Engine/Plugins/Runtime/WorldConditions/Source/WorldConditions/Public/WorldConditionBase.h
@@ -0,0 +1,152 @@
+﻿// Copyright Epic Games, Inc. All Rights Reserved.
+
+#pragma once
+
+#include "WorldConditionTypes.h"
+#include "WorldConditionBase.generated.h"
+
+struct FWorldConditionContext;
+struct FWorldConditionQueryState;
+class UWorldConditionSchema;
+
+/**
+ * Base struct for all World Conditions.
+ *
+ * World Conditions are used together with World Condition Query to create expressions of conditions whose result can be checked.
+ * The conditions can be based on globally accessible data (e.g. subsystems), or based on any of the context data accessible via FWorldConditionContext.
+ * The data that is available is defined for each use case using a UWorldConditionSchema. It also defines which conditions are available when creating
+ * a query in that use case.
+ *
+ * FWorldConditionContextDataRef allows to reference a specific context data. When added on a world condition as property, the UI allows to pick
+ * context data of specified type, and in code, pointer to the actual data can be accessed via the context. 
+ *
+ *		UPROPERTY(EditAnywhere, Category="Default", meta=(BaseStruct="/Script/SmartObjectsModule.SmartObjectSlotHandle"))
+ *		FWorldConditionContextDataRef Slot;
+ *
+ *	EWorldConditionResult FWorldConditionSlotTagQuery::IsTrue(const FWorldConditionContext& Context) const
+ *	{
+ *		if (FSmartObjectSlotHandle* SlotHandle = Context.GetMutableContextDataPtr<FSmartObjectSlotHandle>(Slot))
+ *		{
+ *		}
+ *		...
+ *	}
+ *
+ * Under the hood a reference is a name that needs to be turned into an index before use. This is done on Init().
+ *	bool FWorldConditionSlotTagQuery::Init(const UWorldConditionSchema& Schema)
+ *	{
+ *		if (!Schema.ResolveContextDataRef<FSmartObjectSlotHandle>(Slot))
+ *		{
+ *			return false;
+ *		}
+ *		...
+ *	}
+ *
+ * To speed up query evaluation, the result of a World Condition can be cached by World Condition Query. A condition can be cached, if it is based on
+ * globally accessible data, or context data that is marked as Permanent. To indicate that a conditions result can be cached,
+ * bCanCacheResult must be set to true. This is done on Init().
+ *
+ *	bool FWorldConditionSlotTagQuery::Init(const UWorldConditionSchema& Schema)
+ *	{
+ *		...
+ *		bCanCacheResult = Schema.GetContextDataTypeByRef(Slot) == EWorldConditionContextDataType::Persistent)
+ *		...
+ *	}
+ *
+ * When the result is cached, it needs to be invalidated when new value arrives. This can be done using e.g. a delegate callback which
+ * calls InvalidateResult() of the current condition. This call will invalidate the query state so that next time when IsTrue() is called
+ * required condition will be re-evaluated. It is advised to do as little work as possible in the delegate callback.
+ *
+ *	bool FWorldConditionSlotTagQuery::Activate(const FWorldConditionContext& Context) const
+ *	{
+ *		if (Context.GetContextDataType(Slot) == EWorldConditionContextDataType::Persistent)
+ *		{
+ *			FOnSmartObjectEvent* SlotDelegate = ...;
+ *			FStateType& State = Context.GetState(*this);
+ *			State.DelegateHandle = SlotDelegate->AddLambda([this, &QueryState = Context.GetQueryState()]()
+ *				{
+ *					InvalidateResult(QueryState);
+ *				});
+ *			}
+ *		}
+ *		return false;
+ *	}
+ */
+USTRUCT(meta=(Hidden))
+struct WORLDCONDITIONS_API FWorldConditionBase
+{
+	GENERATED_BODY()
+
+	FWorldConditionBase()
+		: bIsStateObject(false)
+		, bCanCacheResult(true)
+	{
+	}
+	virtual ~FWorldConditionBase();
+
+	/** @return The Instance data type of the condition. */
+	virtual const UStruct* GetRuntimeStateType() const { return nullptr; }
+
+	/** @retrun Description to be shown in the UI. */
+	virtual FText GetDescription() const;
+
+	/**
+	 * Initializes the condition to be used with a specific schema.
+	 * This is called on PostLoad(), or during editing to make sure the data stays in sync with the schema.
+	 * This is the place to resolve all the Context Data References,
+	 * and set bCanCacheResult based on if the context data type.
+	 * @param Schema The schema to initialize the condition for.
+	 * @return True if init succeeded. 
+	 */
+	virtual bool Initialize(const UWorldConditionSchema& Schema);
+
+	/**
+	 * Called to activate the condition.
+	 * The state data for the conditions can be accessed via the Context.
+	 * @param Context Context that stores the context data and state of the query.
+	 * @return True if activation succeeded.
+	 */
+	virtual bool Activate(const FWorldConditionContext& Context) const;
+
+	/**
+	 * Called to check the condition state.
+	 * The state data for the conditions can be accessed via the Context.
+	 * @param Context Context that stores the context data and state of the query.
+	 * @return The state of the condition.
+	 */
+	virtual EWorldConditionResult IsTrue(const FWorldConditionContext& Context) const;
+
+	/**
+	 * Called to deactivate the condition.
+	 * The state data for the conditions can be accessed via the Context.
+	 * @param Context Context that stores the context data and state of the query.
+	 */
+	virtual void Deactivate(const FWorldConditionContext& Context) const;
+
+protected:
+	/** Invalidates the query state so tha the required cached conditions are reevaluated on next call to IsTrue(). */
+	void InvalidateResult(FWorldConditionQueryState& QueryState) const;
+
+	/** Used internally, Offset of the data in the State storage. */
+	uint16 StateDataOffset = 0;
+
+	/** Used internally, Index of the condition in the definition and state storage. */
+	uint8 ConditionIndex = 0;
+
+	/** Used Internally, true if the condition has Object state. */
+	uint8 bIsStateObject : 1;
+
+	/** Set to true if the result of the IsTrue() can be cached. */
+	uint8 bCanCacheResult : 1;
+
+	/** Operator describing how the results of the condition is combined with other conditions. Not used directly, but to set up condition item in query state. */
+	UPROPERTY()
+	EWorldConditionOperator Operator = EWorldConditionOperator::And;
+
+	/** Depth controlling the parenthesis of the expression. Not used directly, but to set up condition item in query state. */
+	UPROPERTY()
+	uint8 NextExpressionDepth = 0;
+
+	friend struct FWorldConditionQueryDefinition;
+	friend struct FWorldConditionQueryState;
+	friend struct FWorldConditionContext;
+};
\ No newline at end of file
diff --git a/Engine/Plugins/Runtime/WorldConditions/Source/WorldConditions/Public/WorldConditionContext.h b/Engine/Plugins/Runtime/WorldConditions/Source/WorldConditions/Public/WorldConditionContext.h
new file mode 100644
--- /dev/null
+++ b/Engine/Plugins/Runtime/WorldConditions/Source/WorldConditions/Public/WorldConditionContext.h
@@ -0,0 +1,252 @@
+﻿// Copyright Epic Games, Inc. All Rights Reserved.
+
+#pragma once
+
+#include "WorldConditionTypes.h"
+#include "WorldConditionQuery.h"
+#include "WorldConditionSchema.h"
+
+/**
+ * The World condition context and context data are structs that are created when we want to interact with world conditions.
+ *
+ * When using FWorldConditionQuery, we only need to deal with context data. The data is defined in 
+ *
+ *	// Create context data for our Fantastic use case.
+ *	const UFantasticConditionSchema* DefaultSchema = GetDefault<UFantasticWorldConditionSchema>();
+ *	FWorldConditionContextData ConditionContextData(*DefaultSchema);
+ *	ConditionContextData.SetContextData(DefaultSchema->GetActorRef(), FantasticActor);
+ *
+ *	if (Query.IsTrue(ConditionContextData))
+ *	{
+ *		...
+ *	}
+ *
+ * When managing separate state memory, world condition context is used to bind everything together:
+ *
+ *	FWorldConditionContext Context(*this, Definition.Preconditions, Runtime.PreconditionState, ConditionContextData);
+ *	if (!Context.IsTrue())
+ *	{
+ *		...
+ *	}
+ *
+ * Note: FWorldConditionContextData and FWorldConditionContext should not be stored for longer durations.
+ *
+ * The expected availability of the world context data is as follows:
+ *
+ *	- Activate:
+ *		- Dynamic: not available
+ *		- Persistent: must be available
+ *	- IsTrue:
+ *		- Dynamic: the passed data might change on each call
+ *		- Persistent: available, but must check if an object is still valid
+ *	- Deactivate:
+ *		- Dynamic: not available
+ *		- Persistent: might not be available
+ *
+ * When using a delegate to invalidate the query, it is advised to store weak pointer or handle to be able to unregister the delegate
+ * even if the persistent data is not available.   
+ */
+
+/**
+ * Container for World Condition context data for a specific schema.
+ */
+struct WORLDCONDITIONS_API FWorldConditionContextData
+{
+	FWorldConditionContextData() = default;
+	
+	explicit FWorldConditionContextData(const UWorldConditionSchema& InSchema)
+	{
+		SetSchema(&InSchema);
+	}
+
+	/** @return True if the context data is initialized with a schema. */
+	bool IsValid() const { return Schema != nullptr; }
+
+	/** @return True if Schema and OtherSchema are valid pointers, and the schema for the context data is child of other context. */
+	bool IsSchemaChildOf(const UWorldConditionSchema* OtherSchema) const
+	{
+		return Schema && OtherSchema && Schema->IsA(OtherSchema->GetClass());
+	}
+
+	/** @return the schema the context data is initialized for. */
+	const UWorldConditionSchema* GetSchema() const { return Schema; }
+
+	/** Sets schema for the context data and initializes data views. */
+	void SetSchema(const UWorldConditionSchema* InSchema)
+	{
+		Schema = InSchema;
+		if (InSchema)
+		{
+			Views.Init(FWorldConditionDataView(), Schema->GetContextDataDescs().Num());
+			for (int32 Index = 0; Index < Views.Num(); Index++)
+			{
+				Views[Index] = FWorldConditionDataView(Schema->GetContextDataTypeByIndex(Index));
+			}
+		}
+		else
+		{
+			Views.Reset();
+		}
+	}
+
+	/** Sets context data Struct at location specified by Ref. */
+	template <typename T>
+	typename TEnableIf<!TIsDerivedFrom<T, UObject>::IsDerived, bool>::Type SetContextData(const FWorldConditionContextDataRef& Ref, const T* Value)
+	{
+		check(Schema);
+		if (Ref.IsValid())
+		{
+			checkSlow(T::StaticStruct()->IsChildOf(Schema->GetContextDataDescByRef(Ref)->Struct));
+			Views[Ref.GetIndex()] = FWorldConditionDataView(T::StaticStruct(), reinterpret_cast<const uint8*>(Value), Schema->GetContextDataTypeByRef(Ref));
+			return true;
+		}
+		return false;
+	}
+
+	/** Sets context data Object at location specified by Ref. */
+	template <typename T>
+	typename TEnableIf<TIsDerivedFrom<T, UObject>::IsDerived, bool>::Type SetContextData(const FWorldConditionContextDataRef& Ref, const T* Object)
+	{
+		check(Schema);
+		if (Ref.IsValid())
+		{
+			checkSlow(T::StaticStruct()->IsChildOf(Schema->GetContextDataDescByRef(Ref)->Struct));
+			Views[Ref.GetIndex()] = FWorldConditionDataView(Object, Schema->GetContextDataTypeByRef(Ref));
+			return true;
+		}
+		return false;
+	}
+
+	/** Sets context data Struct at location specified by Name. */
+	template <typename T>
+	typename TEnableIf<!TIsDerivedFrom<T, UObject>::IsDerived, bool>::Type SetContextData(const FName Name, const T* Value)
+	{
+		check(Schema);
+		const int32 Index = Schema->GetContextDataIndexByName(Name, T::StaticStruct());
+		if (Index != INDEX_NONE)
+		{
+			checkSlow(T::StaticStruct()->IsChildOf(Schema->GetContextDataDescByIndex(Index)->Struct));
+			Views[Index] = FWorldConditionDataView(T::StaticStruct(), reinterpret_cast<const uint8*>(Value), Schema->GetContextDataTypeByIndex(Index));
+			return true;
+		}
+		return false;
+	}
+
+	/** Sets context data Object at location specified by Name. */
+	template <typename T>
+	typename TEnableIf<TIsDerivedFrom<T, UObject>::IsDerived, bool>::Type SetContextData(const FName Name, const T* Object)
+	{
+		check(Schema);
+		const int32 Index = Schema->GetContextDataIndexByName(Name, T::StaticClass());
+		if (Index != INDEX_NONE)
+		{
+			checkSlow(T::StaticStruct()->IsChildOf(Schema->GetContextDataDescByIndex(Index)->Struct));
+			Views[Index] = FWorldConditionDataView(Object, Schema->GetContextDataTypeByIndex(Index));
+			return true;
+		}
+		return false;
+	}
+
+protected:
+	/** Pointer to schema used to initialize the context data. */
+	const UWorldConditionSchema* Schema = nullptr;
+	/** Views to context data. */
+	TArray<FWorldConditionDataView> Views;
+
+	friend struct FWorldConditionContext;
+};
+
+/**
+ * The World Condition context is used to activate, update, and deactivate a world condition.
+ * It ties together the context data, query definition, and query state, and allows data access for the query conditions.
+ */
+struct WORLDCONDITIONS_API FWorldConditionContext
+{
+	explicit FWorldConditionContext(const UObject& Owner, const FWorldConditionQueryDefinition& InQueryDefinition, FWorldConditionQueryState& InQueryState, const FWorldConditionContextData& InContextData)
+		: Owner(Owner)
+		, World(Owner.GetWorld())
+		, QueryDefinition(InQueryDefinition)
+		, QueryState(InQueryState)
+		, ContextData(InContextData)
+	{
+	}
+
+	/** @return Pointer to owner of the world conditions to be updated. */
+	const UObject* GetOwner() const { return &Owner; }
+	
+	/** @return Pointer to world of the owner of the world conditions to be updated. */
+	UWorld* GetWorld() const { return World; }
+	
+	/** @return Pointer to the schema of the context data passed to the conditions. */
+	const UWorldConditionSchema* GetSchema() const { return ContextData.Schema; }
+
+	/** @return Type of the referenced context data. */
+	EWorldConditionContextDataType GetContextDataType(const FWorldConditionContextDataRef& Ref) const
+	{
+		check(Ref.IsValid());
+		return ContextData.Views[Ref.GetIndex()].GetType();
+	}
+
+	/** @todo: Add mutable version with const checking. */
+	/** @return Pointer to referenced context data. */
+	template <typename T>
+	const T* GetContextDataPtr(const FWorldConditionContextDataRef& Ref) const
+	{
+		check(Ref.IsValid());
+		return ContextData.Views[Ref.GetIndex()].template GetPtr<T>();
+	}
+
+	/** @return Struct State data of the specific world condition. */
+	template <typename T>
+	typename T::FStateType& GetState(const T& Condition) const
+	{
+		static_assert(TIsDerivedFrom<T, FWorldConditionBase>::IsDerived, "Expecting Conditions to derive from FWorldConditionBase.");
+		return QueryState.GetStateStruct(Condition).template GetMutable<typename T::FStateType>();
+	}
+
+	/** @return Object State data of the specific world condition. */
+	template <typename T>
+	typename T::UStateType& GetState(const T& Condition) const
+	{
+		static_assert(TIsDerivedFrom<T, FWorldConditionBase>::IsDerived, "Expecting Conditions to derive from FWorldConditionBase.");
+		return *CastChecked<typename T::UStateType>(QueryState.GetStateObject(Condition));
+	}
+
+	/** @return Reference to the query state of update wolrd condition query. */
+	FWorldConditionQueryState& GetQueryState() const { return QueryState; }
+
+	/**
+	 * Calls Activate() on the world conditions in the query.
+	 * @return true of the activation succeeded, or false if failed. Failed queries will return false when IsTrue() is called.
+	 */
+	bool Activate() const;
+
+	/**
+	 * Evaluates the result of the query.
+	 * Intermediate results may be cached, and are stored in the query state.
+	 * If a cached result is invalidated, or the query relies on dynamic context data, IsTrue() is called on the necessary conditions.
+	 * @return the value of the query condition expression.
+	 */
+	bool IsTrue() const;
+
+	/**
+	 * Calls Deactivate() on the world conditions in the query.
+	 */
+	void Deactivate() const;
+	
+protected:
+	/** Owner object of the query, used as outer for duplicate UObject condition state. */
+	const UObject& Owner;
+	
+	/** World of the owner. */
+	UWorld* World = nullptr;
+	
+	/** Reference to the query definition of the query to be updated. */
+	const FWorldConditionQueryDefinition& QueryDefinition;
+	
+	/** Reference to the query state of the query to be updated. */
+	FWorldConditionQueryState& QueryState;
+	
+	/** Reference to the context data for the query to be updated. */
+	const FWorldConditionContextData& ContextData;
+};
\ No newline at end of file
diff --git a/Engine/Plugins/Runtime/WorldConditions/Source/WorldConditions/Public/WorldConditionQuery.h b/Engine/Plugins/Runtime/WorldConditions/Source/WorldConditions/Public/WorldConditionQuery.h
new file mode 100644
--- /dev/null
+++ b/Engine/Plugins/Runtime/WorldConditions/Source/WorldConditions/Public/WorldConditionQuery.h
@@ -0,0 +1,293 @@
+﻿// Copyright Epic Games, Inc. All Rights Reserved.
+
+#pragma once
+
+#include "InstancedStruct.h"
+#include "InstancedStructArray.h"
+#include "StructView.h"
+#include "WorldConditionBase.h"
+#include "WorldConditionTypes.h"
+#include "WorldConditionQuery.generated.h"
+
+struct FWorldConditionContext;
+struct FWorldConditionContextData;
+
+/**
+ * World Condition Query is an expression of World Conditions whose state can be queried.
+ * The state of query and individual conditions can be cached, which allows to evaluate the conditions quickly.
+ * See FWorldConditionBase for more information about the conditions.
+ *
+ * The World Condition Query is split it two parts: FWorldConditionQueryDefinition and FWorldConditionQueryState.
+ * Definition is the "const" part of the query and state contain runtime caching and runtime state of the condition.
+ * This allows the definition to be stored in an asset, and we can allocate just the per instance data when needed.
+ *
+ * Conditions operate on context data which is defined in a UWorldConditionSchema. The schema describes what kind
+ * of structs and objects are available as input for the conditions, and what conditions can be used in specific use case.
+ *
+ * The state is tightly coupled to the definition. The memory layout of the state is stored in the definition.
+ *
+ * For convenience there is also FWorldConditionQuery which combines these two in one package.
+ */
+
+/**
+ * Struct used to store a world condition in editor. Used internally.
+ * Note that the Operator and ExpressionDepth are stored here separately from the World Condition to make sure they are not reset if the Condition is empty. 
+ */
+USTRUCT()
+struct WORLDCONDITIONS_API FWorldConditionEditable
+{
+	GENERATED_BODY()
+
+#if WITH_EDITORONLY_DATA
+	FWorldConditionEditable() = default;
+
+	FWorldConditionEditable(const uint8 InExpressionDepth, const EWorldConditionOperator InOperator, const FConstStructView InCondition)
+		: ExpressionDepth(InExpressionDepth)
+		, Operator(InOperator)
+		, Condition(InCondition)
+	{
+	}
+	
+	void Reset()
+	{
+		Operator = EWorldConditionOperator::And;
+		ExpressionDepth = 0;
+		Condition.Reset();
+	}
+	
+	/** Expression depth controlling the parenthesis of the expression. */
+	UPROPERTY(EditAnywhere, Category="Default")
+	uint8 ExpressionDepth = 0;
+
+	/** Operator describing how the results of the condition is combined with other conditions. */
+	UPROPERTY(EditAnywhere, Category="Default")
+	EWorldConditionOperator Operator = EWorldConditionOperator::And;
+	
+	/** Instance of a world condition. */
+	UPROPERTY(EditAnywhere, Category="Default")
+	FInstancedStruct Condition;
+#endif // WITH_EDITORONLY_DATA
+};
+
+/**
+ * Definition of a world condition.
+ * The mutable state of the world condition is stored in FWorldConditionQueryState.
+ * This allows to reuse the definitions and minimize the runtime memory needed to run queries.
+ */
+USTRUCT()
+struct WORLDCONDITIONS_API FWorldConditionQueryDefinition
+{
+	GENERATED_BODY()
+
+	/** @return true of the definition has conditions and has been initialized. */
+	bool IsValid() const;
+
+	/** Initialized the condition from editable data. */
+	bool Initialize();
+
+	/** Conditions of the query, populated by Initialize(). */
+	UPROPERTY()
+	FInstancedStructArray Conditions;
+
+	/** Schema of the definition. */
+	UPROPERTY()
+	TSubclassOf<UWorldConditionSchema> SchemaClass = nullptr;
+
+#if WITH_EDITORONLY_DATA
+	/** Conditions used while editing, converted in to Conditions via Initialize(). */
+	UPROPERTY(EditAnywhere, Category="Default", meta=(BaseStruct = "/Script/SmartObjectsModule.WorldCondition"))
+	TArray<FWorldConditionEditable> EditableConditions;
+#endif // WITH_EDITORONLY_DATA
+	
+	friend struct FWorldConditionQueryState;
+	friend struct FWorldConditionBase;
+	friend struct FWorldConditionContext;
+};
+
+
+/**
+ * Item used to describe the structure of a world condition query for fast traversal of the expression.
+ */
+USTRUCT()
+struct WORLDCONDITIONS_API FWorldConditionItem
+{
+	GENERATED_BODY()
+
+	FWorldConditionItem() = default;
+	
+	explicit FWorldConditionItem(const EWorldConditionOperator InOperator, const uint8 InNextExpressionDepth)
+		: Operator(InOperator)
+		, NextExpressionDepth(InNextExpressionDepth)
+	{
+	}
+
+	/** Operator describing how the results of the condition is combined with other conditions. */
+	EWorldConditionOperator Operator = EWorldConditionOperator::And;
+	
+	/** Expression depth controlling the parenthesis of the expression. */
+	uint8 NextExpressionDepth = 0;
+	
+	/** Cached result of the condition. */
+	EWorldConditionResult CachedResult = EWorldConditionResult::Invalid;
+};
+
+/**
+ * Struct used to store the pointer to an UObject based condition state.
+ */
+USTRUCT()
+struct WORLDCONDITIONS_API FWorldConditionStateObject
+{
+	GENERATED_BODY()
+
+	UPROPERTY()
+	TObjectPtr<UObject> Object;
+};
+
+
+/**
+ * Runtime state of a world conditions.
+ * The structure of the data for the state is defined in a query definition.
+ * The definition and conditions are stored in FWorldConditionQueryDefinition.
+ * This allows to reuse the definitions and minimize the runtime memory needed to run queries.
+ *
+ * Note: Any code embedding this struct is responsible for calling AddReferencedObjects().
+ */
+struct WORLDCONDITIONS_API FWorldConditionQueryState
+{
+	FWorldConditionQueryState()
+		: bHasPerConditionState(false)
+	{
+	}
+	
+	~FWorldConditionQueryState();
+
+	/** @return True if the state is property initialized. */
+	bool IsValid() const { return Memory != nullptr; }
+
+	/**
+	 * Initialized the state for a specific query definition.
+	 * @param Owner Owner of any objects created during Init().
+	 * @param QueryDefinition definition of the state to initialized.
+	 * @return True if initialization succeeded, false if failed (e.g. invalid definition).
+	 */
+	bool Initialize(const UObject& Owner, const FWorldConditionQueryDefinition& QueryDefinition);
+
+	/**
+	 * Frees the allocated data and objects. The definition must be the same as used in init
+	 * as it is used to traverse the structure in memory.
+	 */
+	void Free(const FWorldConditionQueryDefinition& QueryDefinition);
+
+	/** @return Condition item at specific index */
+	FWorldConditionItem& GetItem(const int32 Index) const
+	{
+		check(Memory && Index >= 0 && Index < (int32)NumConditions);
+		return *(FWorldConditionItem*)(Memory + Index * sizeof(FWorldConditionItem));
+	}
+
+	/** @return Object describing the state of a specified condition. */
+	UObject* GetStateObject(const FWorldConditionBase& Condition) const
+	{
+		check(Condition.StateDataOffset > 0);
+		check(Condition.bIsStateObject);
+		const FWorldConditionStateObject& StateObject = *reinterpret_cast<FWorldConditionStateObject*>(Memory + Condition.StateDataOffset);
+		return StateObject.Object;
+	}
+
+	/** @return struct describing the state of a specified condition. */
+	FStructView GetStateStruct(const FWorldConditionBase& Condition) const
+	{
+		check(Condition.StateDataOffset > 0);
+		check (!Condition.bIsStateObject);
+		const UScriptStruct* ScriptStruct = Cast<UScriptStruct>(Condition.GetRuntimeStateType());
+		check(ScriptStruct);
+		return FStructView(Cast<UScriptStruct>(Condition.GetRuntimeStateType()), Memory + Condition.StateDataOffset);
+	}
+
+	/** @return True if any of the conditions has runtime state. */
+	bool HasPerConditionState() const { return bHasPerConditionState; }
+
+	/** @return The number of conditions in the state data. */
+	int32 GetNumConditions() const { return NumConditions; }
+
+	/** Adds referenced objects to the collector. */
+	void AddReferencedObjects(const FWorldConditionQueryDefinition& QueryDefinition, class FReferenceCollector& Collector) const;
+
+private:
+	EWorldConditionResult CachedResult = EWorldConditionResult::Invalid;
+	uint8 NumConditions = 0;
+	uint8 bHasPerConditionState : 1;
+	uint8* Memory = nullptr;
+
+	friend struct FWorldConditionBase;
+	friend struct FWorldConditionContext;
+};
+
+
+/**
+ * General purpose World Condition Query that combines Query Definition and Query State in one.
+ */
+USTRUCT()
+struct WORLDCONDITIONS_API FWorldConditionQuery
+{
+	GENERATED_BODY()
+
+	~FWorldConditionQuery();
+
+	/** @return True if the query is activated. */
+	bool IsActive() const;
+
+	/**
+	 * Activates the world conditions in the query.
+	 * @param ContextData ContextData that matches the schema of the query.
+	 * @return true of the activation succeeded, or false if failed. Failed queries will return false when IsTrue() is called.
+	 */
+	bool Activate(const UObject& Owner, const FWorldConditionContextData& ContextData);
+
+	/**
+	 * Returns the result of the query. Cached state is returned if it is available,
+	 * if update is needed or the query has dynamic context data, IsTrue() is called on the necessary conditions.
+	 * @param ContextData ContextData that matches the schema of the query.
+	 * @return the value of the query condition expression.
+	 */
+	bool IsTrue(const FWorldConditionContextData& ContextData) const;
+	
+	/**
+	 * Deactivates the world conditions in the query.
+	 * @param ContextData ContextData that matches the schema of the query.
+	 */
+	void Deactivate(const FWorldConditionContextData& ContextData) const;
+
+	/** @return Schema of the query. */
+	const UWorldConditionSchema* GetSchema() const { return QueryDefinition.SchemaClass.GetDefaultObject(); }
+
+	/** Handles object references in the query state. */
+	void AddStructReferencedObjects(class FReferenceCollector& Collector) const;
+
+	/**
+	 * Initializes a query from array of conditions for testing.
+	 * @return true of the query was created and initialized.
+	 */
+	bool DebugInitialize(const TSubclassOf<UWorldConditionSchema> InSchemaClass, const TConstArrayView<FWorldConditionEditable> InConditions);
+
+protected:
+	/** Defines the conditions to run on the query.  */
+	UPROPERTY(EditAnywhere, Category="Default");
+	FWorldConditionQueryDefinition QueryDefinition;
+
+	/** Runtime state of the query. */
+	mutable FWorldConditionQueryState QueryState;
+
+	/** Owner of the query. */
+	UPROPERTY()
+	TObjectPtr<const UObject> Owner = nullptr;
+};
+
+template<>
+struct TStructOpsTypeTraits<FWorldConditionQuery> : public TStructOpsTypeTraitsBase2<FWorldConditionQuery>
+{
+	enum
+	{
+		WithAddStructReferencedObjects = true,
+	};
+};
diff --git a/Engine/Plugins/Runtime/WorldConditions/Source/WorldConditions/Public/WorldConditionSchema.h b/Engine/Plugins/Runtime/WorldConditions/Source/WorldConditions/Public/WorldConditionSchema.h
new file mode 100644
--- /dev/null
+++ b/Engine/Plugins/Runtime/WorldConditions/Source/WorldConditions/Public/WorldConditionSchema.h
@@ -0,0 +1,117 @@
+﻿// Copyright Epic Games, Inc. All Rights Reserved.
+
+#pragma once
+
+#include "EngineDefines.h"
+#include "WorldConditionTypes.h"
+#include "WorldConditionSchema.generated.h"
+
+/**
+ * Describes the context data and allowed world conditions for a specific use case.
+ *
+ * Each schema adds the context data they can provide in their class constructor. This allows
+ * the derived classes to add more data as needed.
+ *
+ *	UCLASS()
+ *	class UFantasticWorldConditionSchema : public UWorldConditionSchema
+ *	{
+ *		GENERATED_BODY()
+ *	public:
+ *		UFantasticWorldConditionSchema::UFantasticWorldConditionSchema(const FObjectInitializer& ObjectInitializer)
+ *			: Super(ObjectInitializer)
+ *		{
+ *			ActorRef = AddContextDataDesc(TEXT("Actor"), AActor::StaticClass(), EWorldConditionContextDataType::Dynamic);
+ *		}
+ *
+ *		// For convenience, when the schema is known, makes it easy to set the context data.
+ *		FWorldConditionContextDataRef GetActorRef() const { return ActorRef; }
+ *
+ *	protected:
+ *		// Filter condition classes that make sense for this context.
+ *		virtual bool IsStructAllowed(const UScriptStruct* InScriptStruct) const override;
+ *		{
+ *			return Super::IsStructAllowed(InScriptStruct)
+ *				|| InScriptStruct->IsChildOf(TBaseStructure<FFantasticWorldConditionBase>::Get());
+ * 		}
+ *		
+ *		FWorldConditionContextDataRef ActorRef;
+ *	};
+*/
+
+
+UCLASS()
+class WORLDCONDITIONS_API UWorldConditionSchema : public UObject
+{
+	GENERATED_BODY()
+
+public:
+
+	/** @return True if world condition of the specific type is allowed. */
+	virtual bool IsStructAllowed(const UScriptStruct* InScriptStruct) const { return false; }
+
+	/** @return All condition descriptors. */
+	TConstArrayView<FWorldConditionContextDataDesc> GetContextDataDescs() const;
+
+	/** @return Context data descriptor of specific name and type. */
+	const FWorldConditionContextDataDesc* GetContextDataDescByName(const FName DataName, const UStruct* Struct) const;
+
+	/** @return Context data descriptor of specific name and type. */
+	const FWorldConditionContextDataDesc* GetContextDataDescByRef(const FWorldConditionContextDataRef& Ref) const;
+
+	/** @return Context data descriptor of specified index. */
+	const FWorldConditionContextDataDesc& GetContextDataDescByIndex(const int32 Index) const;
+
+	/** @return Context Data type based on reference. */
+	EWorldConditionContextDataType GetContextDataTypeByRef(const FWorldConditionContextDataRef& Ref) const;
+
+	/** @return Context Data type based on index. */
+	EWorldConditionContextDataType GetContextDataTypeByIndex(const int32 Index) const;
+
+	/** @return Context data index by name. */
+	int32 GetContextDataIndexByName(const FName DataName, const UStruct* Struct) const;
+
+	/**
+	 * Resolves the index of the context data reference to an Object context data, based on the name in the reference and the type from the template. 
+	 * @param Ref Reference to resolve.
+	 * @return True if the reference was resolved.
+	 */
+	template <typename T>
+	typename TEnableIf<TIsDerivedFrom<T, UObject>::IsDerived, bool>::Type ResolveContextDataRef(FWorldConditionContextDataRef& Ref) const
+	{
+		const int32 Index = GetContextDataIndexByName(Ref.Name, T::StaticClass());
+		if (Index == INDEX_NONE)
+		{
+			Ref.Index = FWorldConditionContextDataRef::InvalidIndex;
+			return false;
+		}
+		Ref.Index = IntCastChecked<uint8>(Index);
+		return true;
+	}
+
+	/**
+	 * Resolves the index of the context data reference to a Struct context data, based on the name in the reference and the type from the template. 
+	 * @param Ref Reference to resolve.
+	 * @return True if the reference was resolved.
+	 */
+	template <typename T>
+	typename TEnableIf<!TIsDerivedFrom<T, UObject>::IsDerived, bool>::Type ResolveContextDataRef(FWorldConditionContextDataRef& Ref) const
+	{
+		const int32 Index = GetContextDataIndexByName(Ref.Name, T::StaticStruct());
+		if (Index == INDEX_NONE)
+		{
+			Ref.Index = FWorldConditionContextDataRef::InvalidIndex;
+			return false;
+		}
+		Ref.Index = IntCastChecked<uint8>(Index);
+		return true;
+	}
+
+protected:
+
+	/** Adds a context data descriptor and returns reference to it. */
+	FWorldConditionContextDataRef AddContextDataDesc(const FName InName, const UStruct* InStruct, const EWorldConditionContextDataType InType);
+
+	/** All context data descriptors. */
+	UPROPERTY()
+	TArray<FWorldConditionContextDataDesc> ContextDataDescs;
+};
diff --git a/Engine/Plugins/Runtime/WorldConditions/Source/WorldConditions/Public/WorldConditionTypes.h b/Engine/Plugins/Runtime/WorldConditions/Source/WorldConditions/Public/WorldConditionTypes.h
new file mode 100644
--- /dev/null
+++ b/Engine/Plugins/Runtime/WorldConditions/Source/WorldConditions/Public/WorldConditionTypes.h
@@ -0,0 +1,222 @@
+﻿// Copyright Epic Games, Inc. All Rights Reserved.
+
+#pragma once
+
+#include "StructView.h"
+#include "GameplayTagContainer.h"
+#include "WorldConditionTypes.generated.h"
+
+WORLDCONDITIONS_API DECLARE_LOG_CATEGORY_EXTERN(LogWorldCondition, Warning, All);
+
+/** Result of a world condition */
+UENUM()
+enum class EWorldConditionResult : uint8
+{
+	/** The result is false. */
+	IsFalse = 0,
+	/** The result is true. */
+	IsTrue = 1,
+	/** The result needs to be recalculated. */
+	Invalid = 2,
+};
+
+/** Boolean operator between conditions in a expression. */
+UENUM()
+enum class EWorldConditionOperator : uint8
+{
+	/** Boolean And */
+	And,
+	/** Boolean Or */
+	Or,
+	/** Used internally, signifies that a right value is copied over left value. */
+	Copy UMETA(Hidden),
+};
+
+namespace UE::WorldCondition
+{
+	/** Maximum depth of sub-expressions in an expression. */
+	constexpr int32 MaxExpressionDepth = 4;
+
+	/** Merges two condition results using specific operator. */
+	inline EWorldConditionResult MergeResults(const EWorldConditionOperator Operator, const EWorldConditionResult InLeft, const EWorldConditionResult InRight)
+	{
+		const uint8 Left = (uint8)InLeft & 1;
+		const uint8 Right = (uint8)InRight & 1;
+		const uint8 ResultAnd = Left & Right;
+		const uint8 ResultOr = Left | Right;
+		const uint8 Result = (Operator == EWorldConditionOperator::And) ? ResultAnd : ResultOr;
+		return (Operator == EWorldConditionOperator::Copy) ? InRight : (EWorldConditionResult)Result;
+	}
+} // UE::WorldCondition
+
+/**
+  * Describes the availability of a context data.
+ */
+UENUM()
+enum class EWorldConditionContextDataType : uint8
+{
+	/** The context data might change on each call, must check for validity, and is only available on IsTrue(). */
+	Dynamic,
+	/** The context data is the same during the lifetime of the condition, must check for validity on Activate()/IsTrue(), might not be available on Deactivate(). */
+	Persistent,
+};
+
+
+/**
+ * Describes context data available for a world conditions.
+ * The descriptors are defined in UWorldConditionSchema, and context data is referred using FWorldConditionContextDataRef.
+ */
+USTRUCT()
+struct WORLDCONDITIONS_API FWorldConditionContextDataDesc
+{
+	GENERATED_BODY()
+
+	FWorldConditionContextDataDesc() = default;
+	explicit FWorldConditionContextDataDesc(const FName InName, const UStruct* InStruct, const EWorldConditionContextDataType InType)
+		: Struct(InStruct)
+		, Name(InName)
+		, Type(InType)
+	{
+	}
+
+	/** Type of the context data, can be a struct or object. */
+	UPROPERTY()
+	TObjectPtr<const UStruct> Struct;
+
+	/** Name of the context data, which is used together with the type struct to find specific data to use. */
+	UPROPERTY()
+	FName Name;
+
+	/** Type of the context data, See EWorldConditionContextDataType. */
+	UPROPERTY()
+	EWorldConditionContextDataType Type = EWorldConditionContextDataType::Dynamic;
+};
+
+/**
+ * Describes a reference to context data defined in UWorldConditionSchema.
+ * When placed as editable property on a World Condition, allows the user to define which context data to use,
+ * and allows quick access of that data via FWorldConditionContext.
+ */
+USTRUCT()
+struct WORLDCONDITIONS_API FWorldConditionContextDataRef
+{
+	GENERATED_BODY()
+
+	/** Marker for invalid index value. */
+	static constexpr uint8 InvalidIndex = 0xff;
+
+	FWorldConditionContextDataRef() = default;
+
+	explicit FWorldConditionContextDataRef(const FName InName, const uint8 InIndex = InvalidIndex)
+		: Name(InName)
+		, Index(InIndex)
+	{
+	}
+	
+	/** @return true if the data ref index is resolved. */
+	bool IsValid() const { return Index != InvalidIndex; } 
+
+	/** @return Referenced context data index. */
+	uint8 GetIndex() const { return Index; }
+
+	/** @return Name of the referenced data. */
+	FName GetName() const { return Name; }
+	
+protected:
+	UPROPERTY()
+	FName Name;
+
+	UPROPERTY()
+	uint8 Index = InvalidIndex;
+
+	friend class UWorldConditionSchema;
+	friend class FWorldConditionContextDataRefDetails;
+};
+
+/**
+ * View holding a pointer to context data struct or object.
+ */
+struct FWorldConditionDataView
+{
+	FWorldConditionDataView() = default;
+
+	explicit FWorldConditionDataView(const EWorldConditionContextDataType InType)
+		: Type(InType)
+	{
+	}
+
+	// USTRUCT() constructor.
+	explicit FWorldConditionDataView(const UScriptStruct* InScriptStruct, const uint8* InMemory, const EWorldConditionContextDataType InType)
+		: Struct(InScriptStruct)
+		, Memory(InMemory)
+		, Type(InType)
+	{
+		// Must have type with valid pointer.
+		check(!Memory || (Memory && Struct));
+	}
+
+	// UOBJECT() constructor.
+	explicit FWorldConditionDataView(const UObject* Object, const EWorldConditionContextDataType InType)
+		: Struct(Object ? Object->GetClass() : nullptr)
+		, Memory(reinterpret_cast<const uint8*>(Object))
+		, Type(InType)
+	{
+		// Must have type with valid pointer.
+		check(!Memory || (Memory && Struct));
+	}
+
+	// USTRUCT() from a StructView.
+	explicit FWorldConditionDataView(FConstStructView StructView, const EWorldConditionContextDataType InType)
+		: Struct(StructView.GetScriptStruct())
+		, Memory(StructView.GetMemory())
+		, Type(InType)
+	{
+		// Must have type with valid pointer.
+		check(!Memory || (Memory && Struct));
+	}
+
+	/** @return True if the view is set to valid data. */
+	bool IsValid() const
+	{
+		return Memory != nullptr && Struct != nullptr;
+	}
+
+	/** @return Pointer to object stored in the view. */
+	template <typename T>
+	typename TEnableIf<TIsDerivedFrom<T, UObject>::IsDerived, const T*>::Type GetPtr() const
+	{
+		// If Memory is set, expect Struct too. Otherwise, let nulls pass through.
+		check(!Memory || (Memory && Struct));
+		check(!Struct || Struct->IsChildOf(T::StaticClass()));
+		return ((T*)Memory);
+	}
+
+	/** @return Pointer to struct stored in the view. */
+	template <typename T>
+    typename TEnableIf<!TIsDerivedFrom<T, UObject>::IsDerived, const T*>::Type GetPtr() const
+	{
+		// If Memory is set, expect Struct too. Otherwise, let nulls pass through.
+		check(!Memory || (Memory && Struct));
+		check(!Struct || Struct->IsChildOf(T::StaticStruct()));
+		return ((T*)Memory);
+	}
+
+	/** @return Struct describing the data type. */
+	const UStruct* GetStruct() const { return Struct; }
+
+	/** @return Raw const pointer to the data. */
+	const uint8* GetMemory() const { return Memory; }
+
+	/** @return type of the context data. */
+	EWorldConditionContextDataType GetType() const { return Type; }
+	
+protected:
+	/** UClass or UScriptStruct of the data. */
+	const UStruct* Struct = nullptr;
+
+	/** Memory pointing at the class or struct */
+	const uint8* Memory = nullptr;
+
+	/** Type of the context data, copied from the descriptor. */
+	EWorldConditionContextDataType Type = EWorldConditionContextDataType::Dynamic;
+};
diff --git a/Engine/Plugins/Runtime/WorldConditions/Source/WorldConditions/Public/WorldConditionsModule.h b/Engine/Plugins/Runtime/WorldConditions/Source/WorldConditions/Public/WorldConditionsModule.h
new file mode 100644
--- /dev/null
+++ b/Engine/Plugins/Runtime/WorldConditions/Source/WorldConditions/Public/WorldConditionsModule.h
@@ -0,0 +1,37 @@
+// Copyright Epic Games, Inc. All Rights Reserved.
+
+#pragma  once
+
+#include "CoreMinimal.h"
+#include "Modules/ModuleManager.h"
+
+/**
+* The public interface to this module
+*/
+class IWorldConditionsModule : public IModuleInterface
+{
+
+public:
+
+	/**
+	* Singleton-like access to this module's interface.  This is just for convenience!
+	* Beware of calling this during the shutdown phase, though.  Your module might have been unloaded already.
+	*
+	* @return Returns singleton instance, loading the module on demand if needed
+	*/
+	static inline IWorldConditionsModule& Get()
+	{
+		return FModuleManager::LoadModuleChecked<IWorldConditionsModule>("WorldConditions");
+	}
+
+	/**
+	* Checks to see if this module is loaded and ready.  It is only valid to call Get() if IsAvailable() returns true.
+	*
+	* @return True if the module is loaded and ready to use
+	*/
+	static inline bool IsAvailable()
+	{
+		return FModuleManager::Get().IsModuleLoaded("WorldConditions");
+	}
+};
+
diff --git a/Engine/Plugins/Runtime/WorldConditions/Source/WorldConditions/WorldConditions.Build.cs b/Engine/Plugins/Runtime/WorldConditions/Source/WorldConditions/WorldConditions.Build.cs
new file mode 100644
--- /dev/null
+++ b/Engine/Plugins/Runtime/WorldConditions/Source/WorldConditions/WorldConditions.Build.cs
@@ -0,0 +1,25 @@
+// Copyright Epic Games, Inc. All Rights Reserved.
+
+namespace UnrealBuildTool.Rules
+{
+	public class WorldConditions : ModuleRules
+	{
+		public WorldConditions(ReadOnlyTargetRules Target) : base(Target)
+		{
+			PublicIncludePaths.AddRange(
+				new string[] {
+				}
+			);
+
+			PublicDependencyModuleNames.AddRange(
+				new[] {
+					"Core",
+					"CoreUObject",
+					"Engine",
+					"GameplayTags",
+					"StructUtils",
+				}
+			);
+		}
+	}
+}
diff --git a/Engine/Plugins/Runtime/WorldConditions/Source/WorldConditionsEditor/Private/Customizations/WorldConditionContextDataRefDetails.cpp b/Engine/Plugins/Runtime/WorldConditions/Source/WorldConditionsEditor/Private/Customizations/WorldConditionContextDataRefDetails.cpp
new file mode 100644
--- /dev/null
+++ b/Engine/Plugins/Runtime/WorldConditions/Source/WorldConditionsEditor/Private/Customizations/WorldConditionContextDataRefDetails.cpp
@@ -0,0 +1,196 @@
+﻿// Copyright Epic Games, Inc. All Rights Reserved.
+
+#include "WorldConditionContextDataRefDetails.h"
+#include "PropertyHandle.h"
+#include "IDetailChildrenBuilder.h"
+#include "DetailWidgetRow.h"
+#include "PropertyCustomizationHelpers.h"
+#include "Widgets/Input/SComboButton.h"
+#include "DetailLayoutBuilder.h"
+#include "IPropertyUtilities.h"
+#include "ScopedTransaction.h"
+#include "Framework/MultiBox/MultiBoxBuilder.h"
+#include "Widgets/Input/SComboButton.h"
+#include "WorldConditionTypes.h"
+#include "WorldConditionQuery.h"
+#include "WorldConditionSchema.h"
+
+#define LOCTEXT_NAMESPACE "WorldCondition"
+
+TSharedRef<IPropertyTypeCustomization> FWorldConditionContextDataRefDetails::MakeInstance()
+{
+	return MakeShareable(new FWorldConditionContextDataRefDetails);
+}
+
+void FWorldConditionContextDataRefDetails::CustomizeHeader(TSharedRef<IPropertyHandle> StructPropertyHandle, FDetailWidgetRow& HeaderRow, IPropertyTypeCustomizationUtils& StructCustomizationUtils)
+{
+	StructProperty = StructPropertyHandle;
+	PropUtils = StructCustomizationUtils.GetPropertyUtilities().Get();
+
+	CacheContextData();
+	
+	HeaderRow
+		.NameContent()
+		[
+			StructPropertyHandle->CreatePropertyNameWidget()
+		]
+		.ValueContent()
+		.VAlign(VAlign_Center)
+		[
+			SNew(SComboButton)
+			.OnGetMenuContent(this, &FWorldConditionContextDataRefDetails::OnGetContent)
+			.ContentPadding(FMargin(6.f, 0.f))
+			.ButtonContent()
+			[
+				SNew(STextBlock)
+				.Text(this, &FWorldConditionContextDataRefDetails::GetCurrentDesc)
+				.Font(IDetailLayoutBuilder::GetDetailFont())
+			]
+		];
+}
+
+void FWorldConditionContextDataRefDetails::CustomizeChildren(TSharedRef<IPropertyHandle> StructPropertyHandle, IDetailChildrenBuilder& StructBuilder, IPropertyTypeCustomizationUtils& StructCustomizationUtils)
+{
+}
+
+void FWorldConditionContextDataRefDetails::CacheContextData()
+{
+	// Base struct
+	static const FName BaseStructMetaName(TEXT("BaseStruct")); // TODO: move these names into one central place.
+	static const FName BaseClassMetaName(TEXT("BaseClass")); // TODO: move these names into one central place.
+	
+	const FString BaseStructName = StructProperty->GetMetaData(BaseStructMetaName);
+	BaseStruct = UClass::TryFindTypeSlow<UScriptStruct>(BaseStructName);
+
+	if (BaseStruct == nullptr)
+	{
+		const FString BaseClassName = StructProperty->GetMetaData(BaseClassMetaName);
+		BaseStruct = UClass::TryFindTypeSlow<UClass>(BaseClassName);
+	}
+
+	// Find schema from outer FWorldConditionQueryDefinition.
+	Schema = nullptr;
+	
+	TSharedPtr<IPropertyHandle> CurrentProperty = StructProperty;
+	while (CurrentProperty.IsValid() && !Schema)
+	{
+		const FProperty* Property = CurrentProperty->GetProperty();
+		if (const FStructProperty* CurrentStructProperty = CastField<FStructProperty>(Property))
+		{
+			if (CurrentStructProperty->Struct == TBaseStructure<FWorldConditionQueryDefinition>::Get())
+			{
+				// Get schema from definition
+				TArray<void*> RawNodeData;
+				CurrentProperty->AccessRawData(RawNodeData);
+				for (void* Data : RawNodeData)
+				{
+					if (const FWorldConditionQueryDefinition* QueryDefinition = static_cast<FWorldConditionQueryDefinition*>(Data))
+					{
+						Schema = QueryDefinition->SchemaClass.GetDefaultObject();
+						if (Schema)
+						{
+							break;
+						}
+					}
+				}
+			}
+		}
+
+		CurrentProperty = CurrentProperty->GetParentHandle(); 
+	}
+
+	if (Schema && BaseStruct)
+	{
+		CachedContextData.Reset();
+		for (const FWorldConditionContextDataDesc& Desc : Schema->GetContextDataDescs())
+		{
+			if (Desc.Struct->IsChildOf(BaseStruct))
+			{
+				CachedContextData.Add(Desc.Name);
+			}
+		}
+	}
+}
+
+TSharedRef<SWidget> FWorldConditionContextDataRefDetails::OnGetContent() const
+{
+	FMenuBuilder MenuBuilder(true, NULL);
+
+	for (const FName& ContextDataName : CachedContextData)
+	{
+		FUIAction ItemAction(FExecuteAction::CreateSP(const_cast<FWorldConditionContextDataRefDetails*>(this), &FWorldConditionContextDataRefDetails::OnSetContextData, ContextDataName));
+		MenuBuilder.AddMenuEntry(FText::FromName(ContextDataName), FText::GetEmpty(), FSlateIcon(), ItemAction);
+	}
+
+	return MenuBuilder.MakeWidget();
+}
+
+FText FWorldConditionContextDataRefDetails::GetCurrentDesc() const
+{
+	if (!Schema)
+	{
+		return FText::GetEmpty();
+	}
+	
+	if (const FWorldConditionContextDataRef* DataRef = GetCommonContextDataRef())
+	{
+		if (const FWorldConditionContextDataDesc* Desc = Schema->GetContextDataDescByName(DataRef->Name, BaseStruct))
+		{
+			return FText::FromName(Desc->Name);
+		}
+		return LOCTEXT("Invalid", "Invalid");
+	}
+	
+	return LOCTEXT("MultipleSelected", "Multiple Selected");
+}
+
+FWorldConditionContextDataRef* FWorldConditionContextDataRefDetails::GetCommonContextDataRef() const
+{
+	check(StructProperty);
+
+	FWorldConditionContextDataRef* Result = nullptr;
+	
+	TArray<void*> RawNodeData;
+	StructProperty->AccessRawData(RawNodeData);
+
+	for (void* Data : RawNodeData)
+	{
+		if (Data)
+		{
+			Result = static_cast<FWorldConditionContextDataRef*>(Data);
+			break;
+		}
+	}
+
+	return Result;
+}
+
+void FWorldConditionContextDataRefDetails::OnSetContextData(const FName ContextDataName) const
+{
+	check(StructProperty);
+
+	TArray<void*> RawNodeData;
+	StructProperty->AccessRawData(RawNodeData);
+
+	GEditor->BeginTransaction(LOCTEXT("SelectReference", "Select Reference"));
+
+	StructProperty->NotifyPreChange();
+
+	for (void* Data : RawNodeData)
+	{
+		if (FWorldConditionContextDataRef* DataRef = static_cast<FWorldConditionContextDataRef*>(Data))
+		{
+			DataRef->Name = ContextDataName;
+			DataRef->Index = FWorldConditionContextDataRef::InvalidIndex;
+		}
+	}
+
+	StructProperty->NotifyPostChange(EPropertyChangeType::ValueSet);
+	StructProperty->NotifyFinishedChangingProperties();
+
+	GEditor->EndTransaction();
+
+}
+
+
+#undef LOCTEXT_NAMESPACE
diff --git a/Engine/Plugins/Runtime/WorldConditions/Source/WorldConditionsEditor/Private/Customizations/WorldConditionContextDataRefDetails.h b/Engine/Plugins/Runtime/WorldConditions/Source/WorldConditionsEditor/Private/Customizations/WorldConditionContextDataRefDetails.h
new file mode 100644
--- /dev/null
+++ b/Engine/Plugins/Runtime/WorldConditions/Source/WorldConditionsEditor/Private/Customizations/WorldConditionContextDataRefDetails.h
@@ -0,0 +1,45 @@
+﻿// Copyright Epic Games, Inc. All Rights Reserved.
+
+#pragma once
+
+#include "PropertyHandle.h"
+#include "IDetailCustomization.h"
+#include "IPropertyTypeCustomization.h"
+
+class IPropertyHandle;
+class IPropertyUtilities;
+class USmartObjectDefinition;
+class SWidget;
+struct FWorldConditionContextDataRef;
+class UWorldConditionSchema;
+
+/**
+ * Type customization for FWorldConditionContextDataRef.
+ */
+
+class FWorldConditionContextDataRefDetails : public IPropertyTypeCustomization
+{
+public:
+	/** Makes a new instance of this detail layout class for a specific detail view requesting it */
+	static TSharedRef<IPropertyTypeCustomization> MakeInstance();
+
+private:
+	/** IPropertyTypeCustomization interface */
+	virtual void CustomizeHeader(TSharedRef<IPropertyHandle> StructPropertyHandle, FDetailWidgetRow& HeaderRow, IPropertyTypeCustomizationUtils& StructCustomizationUtils) override;
+	virtual void CustomizeChildren(TSharedRef<IPropertyHandle> StructPropertyHandle, IDetailChildrenBuilder& StructBuilder, IPropertyTypeCustomizationUtils& StructCustomizationUtils) override;
+
+	void CacheContextData();
+	TSharedRef<SWidget> OnGetContent() const;
+	FText GetCurrentDesc() const;
+
+	void OnSetContextData(const FName ContextDataName) const;
+	FWorldConditionContextDataRef* GetCommonContextDataRef() const;
+	
+	TArray<FName> CachedContextData;
+
+	const UWorldConditionSchema* Schema = nullptr;
+	const UStruct* BaseStruct = nullptr;
+	
+	IPropertyUtilities* PropUtils = nullptr;
+	TSharedPtr<IPropertyHandle> StructProperty;
+};
diff --git a/Engine/Plugins/Runtime/WorldConditions/Source/WorldConditionsEditor/Private/Customizations/WorldConditionEditableDetails.cpp b/Engine/Plugins/Runtime/WorldConditions/Source/WorldConditionsEditor/Private/Customizations/WorldConditionEditableDetails.cpp
new file mode 100644
--- /dev/null
+++ b/Engine/Plugins/Runtime/WorldConditions/Source/WorldConditionsEditor/Private/Customizations/WorldConditionEditableDetails.cpp
@@ -0,0 +1,657 @@
+﻿// Copyright Epic Games, Inc. All Rights Reserved.
+
+#include "WorldConditionEditableDetails.h"
+#include "PropertyHandle.h"
+#include "IDetailChildrenBuilder.h"
+#include "IPropertyUtilities.h"
+#include "DetailWidgetRow.h"
+#include "PropertyCustomizationHelpers.h"
+#include "Widgets/Input/SComboButton.h"
+#include "DetailLayoutBuilder.h"
+#include "WorldConditionQuery.h"
+#include "ScopedTransaction.h"
+#include "Framework/MultiBox/MultiBoxBuilder.h"
+#include "Styling/StyleColors.h"
+#include "StructViewerModule.h"
+#include "StructViewerFilter.h"
+#include "WorldConditionEditorStyle.h"
+#include "InstancedStructDetails.h"
+#include "WorldConditionSchema.h"
+
+#define LOCTEXT_NAMESPACE "WorldCondition"
+
+class FWorldConditionStructFilter : public IStructViewerFilter
+{
+public:
+	/** The base struct for the property that classes must be a child-of. */
+	const UScriptStruct* BaseStruct = nullptr;
+	const UWorldConditionSchema* Schema = nullptr;
+
+	virtual bool IsStructAllowed(const FStructViewerInitializationOptions& InInitOptions, const UScriptStruct* InStruct, TSharedRef<FStructViewerFilterFuncs> InFilterFuncs) override
+	{
+		if (!InStruct || InStruct == BaseStruct)
+		{
+			return false;
+		}
+
+		if (Schema && !Schema->IsStructAllowed(InStruct))
+		{
+			return false;
+		}
+
+		if (InStruct->HasMetaData(TEXT("Hidden")))
+		{
+			return false;
+		}
+
+		// Query the native struct to see if it has the correct parent type (if any)
+		return !BaseStruct || InStruct->IsChildOf(BaseStruct);
+	}
+
+	virtual bool IsUnloadedStructAllowed(const FStructViewerInitializationOptions& InInitOptions, const FSoftObjectPath& InStructPath, TSharedRef<FStructViewerFilterFuncs> InFilterFuncs) override
+	{
+		return false;
+	}
+};
+
+
+TSharedRef<IPropertyTypeCustomization> FWorldConditionEditableDetails::MakeInstance()
+{
+	return MakeShareable(new FWorldConditionEditableDetails);
+}
+
+void FWorldConditionEditableDetails::CustomizeHeader(TSharedRef<IPropertyHandle> StructPropertyHandle, FDetailWidgetRow& HeaderRow, IPropertyTypeCustomizationUtils& StructCustomizationUtils)
+{
+	StructProperty = StructPropertyHandle;
+	PropUtils = StructCustomizationUtils.GetPropertyUtilities().Get();
+
+	ConditionProperty = StructProperty->GetChildHandle(TEXT("Condition"));
+	ExpressionDepthProperty = StructProperty->GetChildHandle(TEXT("ExpressionDepth"));
+	OperatorProperty = StructProperty->GetChildHandle(TEXT("Operator"));
+
+	check(ConditionProperty.IsValid());
+	check(ExpressionDepthProperty.IsValid());
+	check(OperatorProperty.IsValid());
+	
+	CacheSchema();
+
+	const FIsResetToDefaultVisible IsResetVisible = FIsResetToDefaultVisible::CreateSP(this, &FWorldConditionEditableDetails::ShouldResetToDefault);
+	const FResetToDefaultHandler ResetHandler = FResetToDefaultHandler::CreateSP(this, &FWorldConditionEditableDetails::ResetToDefault);
+	const FResetToDefaultOverride ResetOverride = FResetToDefaultOverride::Create(IsResetVisible, ResetHandler);
+
+	HeaderRow
+		.WholeRowContent()
+		.VAlign(VAlign_Center)
+		[
+			SNew(SHorizontalBox)
+			+ SHorizontalBox::Slot()
+			.FillWidth(2.0f)
+			[
+				SNew(SHorizontalBox)
+				// Indent
+				+ SHorizontalBox::Slot()
+				.AutoWidth()
+				.VAlign(VAlign_Center)
+				[
+					SNew(SBox)
+					.WidthOverride(this, &FWorldConditionEditableDetails::GetIndentSize)
+					[
+						SNew(SComboButton)
+						.ComboButtonStyle(FWorldConditionEditorStyle::Get(), "Condition.Indent.ComboBox")
+						.ContentPadding(2.0f)
+						.HasDownArrow(false)
+						.VAlign(VAlign_Center)
+						.OnGetMenuContent(this, &FWorldConditionEditableDetails::OnGetIndentContent)
+						.ButtonContent()
+						[
+							SNew(STextBlock)
+							.Text(FText::FromString(TEXT(">")))
+							.Font(IDetailLayoutBuilder::GetDetailFont())
+							.ToolTipText(LOCTEXT("DepthTooltip", "Depth of the expression row, controls parentheses and evaluation order."))
+						]
+					]
+				]
+				// Operator
+				+ SHorizontalBox::Slot()
+				.AutoWidth()
+				.VAlign(VAlign_Center)
+				[
+					SNew(SBox)
+					.WidthOverride(30.0f)
+					.Padding(FMargin(2, 4))
+					.VAlign(VAlign_Center)
+					[
+						SNew(SComboButton)
+						.IsEnabled(TAttribute<bool>(this, &FWorldConditionEditableDetails::IsOperatorEnabled))
+						.ComboButtonStyle(FWorldConditionEditorStyle::Get(), "Condition.Operator.ComboBox")
+						.ButtonColorAndOpacity(this, &FWorldConditionEditableDetails::GetOperatorColor)
+						.HasDownArrow(false)
+						.HAlign(HAlign_Center)
+						.VAlign(VAlign_Center)
+						.OnGetMenuContent(this, &FWorldConditionEditableDetails::OnGetOperatorContent)
+						.ButtonContent()
+						[
+							SNew(STextBlock)
+							.TextStyle(FWorldConditionEditorStyle::Get(), "Condition.Operator")
+							.Text(this, &FWorldConditionEditableDetails::GetOperatorText)
+						]
+					]
+				]
+				// Open parens
+				+ SHorizontalBox::Slot()
+				.AutoWidth()
+				.VAlign(VAlign_Center)
+				[
+					SNew(STextBlock)
+					.TextStyle(FWorldConditionEditorStyle::Get(), "Condition.Parens")
+					.Text(this, &FWorldConditionEditableDetails::GetOpenParens)
+				]
+
+				// Class picker
+				+ SHorizontalBox::Slot()
+				.FillWidth(1.0f)
+				.Padding(FMargin(FMargin(4.0f, 0.0f, 0.0f, 0.0f)))
+				.VAlign(VAlign_Center)
+				[
+					SAssignNew(ComboButton, SComboButton)
+					.OnGetMenuContent(this, &FWorldConditionEditableDetails::GeneratePicker)
+					.ContentPadding(0.f)
+					.ButtonContent()
+					[
+						SNew(STextBlock)
+						.Text(this, &FWorldConditionEditableDetails::GetDisplayValueString)
+						.Font(IDetailLayoutBuilder::GetDetailFont())
+					]
+				]
+
+				// Close parens
+				+ SHorizontalBox::Slot()
+				.AutoWidth()
+				.VAlign(VAlign_Center)
+				[
+					SNew(STextBlock)
+					.TextStyle(FWorldConditionEditorStyle::Get(), "Condition.Parens")
+					.Text(this, &FWorldConditionEditableDetails::GetCloseParens)
+				]
+			]
+			+ SHorizontalBox::Slot()
+			.AutoWidth()
+			[
+				StructPropertyHandle->CreateDefaultPropertyButtonWidgets()
+			]
+		]
+		.OverrideResetToDefault(ResetOverride);
+}
+
+void FWorldConditionEditableDetails::CustomizeChildren(TSharedRef<IPropertyHandle> StructPropertyHandle, IDetailChildrenBuilder& StructBuilder, IPropertyTypeCustomizationUtils& StructCustomizationUtils)
+{
+	check(ConditionProperty);
+	
+	TSharedRef<FInstancedStructDataDetails> NodeDetails = MakeShareable(new FInstancedStructDataDetails(ConditionProperty));
+	StructBuilder.AddCustomBuilder(NodeDetails);
+}
+
+bool FWorldConditionEditableDetails::ShouldResetToDefault(TSharedPtr<IPropertyHandle> PropertyHandle) const
+{
+	check(StructProperty);
+	
+	bool bAnyValid = false;
+	
+	TArray<const void*> RawNodeData;
+	StructProperty->AccessRawData(RawNodeData);
+	for (const void* Data : RawNodeData)
+	{
+		if (const FWorldConditionEditable* Condition = static_cast<const FWorldConditionEditable*>(Data))
+		{
+			if (Condition->Condition.IsValid())
+			{
+				bAnyValid = true;
+				break;
+			}
+		}
+	}
+	
+	// Assume that the default value is empty. Any valid means that some can be reset to empty.
+	return bAnyValid;
+}
+
+
+void FWorldConditionEditableDetails::ResetToDefault(TSharedPtr<IPropertyHandle> PropertyHandle)
+{
+	check(StructProperty);
+	
+	GEditor->BeginTransaction(LOCTEXT("OnResetToDefault", "Reset to default"));
+
+	StructProperty->NotifyPreChange();
+	
+	TArray<void*> RawNodeData;
+	StructProperty->AccessRawData(RawNodeData);
+	for (void* Data : RawNodeData)
+	{
+		if (FWorldConditionEditable* Node = static_cast<FWorldConditionEditable*>(Data))
+		{
+			Node->Reset();
+		}
+	}
+
+	StructProperty->NotifyPostChange(EPropertyChangeType::ValueSet);
+	StructProperty->NotifyFinishedChangingProperties();
+
+	GEditor->EndTransaction();
+
+	if (PropUtils)
+	{
+		PropUtils->ForceRefresh();
+	}
+}
+
+void FWorldConditionEditableDetails::CacheSchema()
+{
+	// Find schema from outer FWorldConditionQueryDefinition.
+	Schema = nullptr;
+	
+	TSharedPtr<IPropertyHandle> CurrentProperty = StructProperty;
+	while (CurrentProperty.IsValid() && !Schema)
+	{
+		const FProperty* Property = CurrentProperty->GetProperty();
+		if (const FStructProperty* CurrentStructProperty = CastField<FStructProperty>(Property))
+		{
+			if (CurrentStructProperty->Struct == TBaseStructure<FWorldConditionQueryDefinition>::Get())
+			{
+				// Get schema from definition
+				TArray<void*> RawNodeData;
+				CurrentProperty->AccessRawData(RawNodeData);
+				for (void* Data : RawNodeData)
+				{
+					if (const FWorldConditionQueryDefinition* QueryDefinition = static_cast<FWorldConditionQueryDefinition*>(Data))
+					{
+						Schema = QueryDefinition->SchemaClass.GetDefaultObject();
+						if (Schema)
+						{
+							break;
+						}
+					}
+				}
+			}
+		}
+		CurrentProperty = CurrentProperty->GetParentHandle(); 
+	}
+}
+
+FOptionalSize FWorldConditionEditableDetails::GetIndentSize() const
+{
+	return FOptionalSize(15.0f + GetExpressionDepth() * 30.0f);
+}
+
+TSharedRef<SWidget> FWorldConditionEditableDetails::OnGetIndentContent() const
+{
+	FMenuBuilder MenuBuilder(true, nullptr);
+
+	for (int32 Depth = 0; Depth < UE::WorldCondition::MaxExpressionDepth; Depth++)
+	{
+		FUIAction ItemAction(
+			FExecuteAction::CreateSP(this, &FWorldConditionEditableDetails::SetExpressionDepth, Depth),
+			FCanExecuteAction(),
+			FIsActionChecked::CreateSP(this, &FWorldConditionEditableDetails::IsExpressionDepth, Depth));
+		MenuBuilder.AddMenuEntry(FText::AsNumber(Depth), TAttribute<FText>(), FSlateIcon(), ItemAction, FName(), EUserInterfaceActionType::Check);
+	}
+
+	return MenuBuilder.MakeWidget();
+}
+
+int32 FWorldConditionEditableDetails::GetExpressionDepth() const
+{
+	check(ExpressionDepthProperty);
+	
+	uint8 Depth = 0;
+	ExpressionDepthProperty->GetValue(Depth);
+
+	return Depth;
+}
+
+void FWorldConditionEditableDetails::SetExpressionDepth(const int32 Depth) const
+{
+	check(ExpressionDepthProperty);
+	
+	ExpressionDepthProperty->SetValue((uint8)FMath::Clamp(Depth, 0, UE::WorldCondition::MaxExpressionDepth - 1));
+}
+
+bool FWorldConditionEditableDetails::IsExpressionDepth(const int32 Depth) const
+{
+	return Depth == GetExpressionDepth();
+}
+
+FText FWorldConditionEditableDetails::GetOperatorText() const
+{
+	check(OperatorProperty);
+
+	// First item does not relate to anything existing, it could be empty, we return IF to indicate that we're building condition. 
+	if (IsFirstItem())
+	{
+		return LOCTEXT("IfOperator", "IF");
+	}
+
+	uint8 Value = 0;
+	OperatorProperty->GetValue(Value);
+	const EWorldConditionOperator Operator = (EWorldConditionOperator)Value;
+
+	if (Operator == EWorldConditionOperator::And)
+	{
+		return LOCTEXT("AndOperator", "AND");
+	}
+	else if (Operator == EWorldConditionOperator::Or)
+	{
+		return LOCTEXT("OrOperator", "OR");
+	}
+	else
+	{
+		ensureMsgf(false, TEXT("Unhandled Operator %s"), *UEnum::GetValueAsString(Operator));
+	}
+
+	return FText::GetEmpty();
+}
+
+FSlateColor FWorldConditionEditableDetails::GetOperatorColor() const
+{
+	check(OperatorProperty);
+
+	if (IsFirstItem())
+	{
+		return FStyleColors::Transparent;
+	}
+
+	uint8 Value = 0; 
+	OperatorProperty->GetValue(Value);
+	const EWorldConditionOperator Operator = (EWorldConditionOperator)Value;
+
+	if (Operator == EWorldConditionOperator::And)
+	{
+		return FStyleColors::AccentPink;
+	}
+	else if (Operator == EWorldConditionOperator::Or)
+	{
+		return FStyleColors::AccentBlue;
+	}
+	else
+	{
+		ensureMsgf(false, TEXT("Unhandled Operator %s"), *UEnum::GetValueAsString(Operator));
+	}
+
+	return FStyleColors::Transparent;
+
+}
+
+TSharedRef<SWidget> FWorldConditionEditableDetails::OnGetOperatorContent() const
+{
+	FMenuBuilder MenuBuilder(true, NULL);
+
+	FUIAction AndAction(
+		FExecuteAction::CreateSP(this, &FWorldConditionEditableDetails::SetOperator, EWorldConditionOperator::And),
+		FCanExecuteAction(),
+		FIsActionChecked::CreateSP(this, &FWorldConditionEditableDetails::IsOperator, EWorldConditionOperator::And));
+	MenuBuilder.AddMenuEntry(LOCTEXT("AndOperator", "AND"), TAttribute<FText>(), FSlateIcon(), AndAction, FName(), EUserInterfaceActionType::Check);
+
+	FUIAction OrAction(FExecuteAction::CreateSP(this, &FWorldConditionEditableDetails::SetOperator, EWorldConditionOperator::Or),
+		FCanExecuteAction(),
+		FIsActionChecked::CreateSP(this, &FWorldConditionEditableDetails::IsOperator, EWorldConditionOperator::Or));
+	MenuBuilder.AddMenuEntry(LOCTEXT("OrOperator", "OR"), TAttribute<FText>(), FSlateIcon(), OrAction, FName(), EUserInterfaceActionType::Check);
+
+	return MenuBuilder.MakeWidget();
+}
+
+bool FWorldConditionEditableDetails::IsOperatorEnabled() const
+{
+	return !IsFirstItem();
+}
+
+void FWorldConditionEditableDetails::SetOperator(const EWorldConditionOperator Operator) const
+{
+	check(OperatorProperty);
+	
+	OperatorProperty->SetValue((uint8)Operator);
+}
+
+bool FWorldConditionEditableDetails::IsOperator(const EWorldConditionOperator Operator) const
+{
+	check(OperatorProperty);
+
+	uint8 Value = 0; 
+	OperatorProperty->GetValue(Value);
+	const EWorldConditionOperator CurrOperator = (EWorldConditionOperator)Value;
+
+	return CurrOperator == Operator;
+}
+
+bool FWorldConditionEditableDetails::IsFirstItem() const
+{
+	check(StructProperty);
+	return StructProperty->GetIndexInArray() == 0;
+}
+
+int32 FWorldConditionEditableDetails::GetCurrExpressionDepth() const
+{
+	// First item needs to be zero depth to make the parentheses counting to work properly.
+	return IsFirstItem() ? 0 : GetExpressionDepth();
+}
+
+int32 FWorldConditionEditableDetails::GetNextExpressionDepth() const
+{
+	// Find the intent of the next item by finding the item in the parent array.
+	check(StructProperty);
+	TSharedPtr<IPropertyHandle> ParentProp = StructProperty->GetParentHandle();
+	if (!ParentProp.IsValid())
+	{
+		return 0;
+	}
+	TSharedPtr<IPropertyHandleArray> ParentArray = ParentProp->AsArray();
+	if (!ParentArray.IsValid())
+	{
+		return 0;
+	}
+
+	uint32 NumElements = 0;
+	if (ParentArray->GetNumElements(NumElements) != FPropertyAccess::Success)
+	{
+		return 0;
+	}
+	
+	const int32 NextIndex = StructProperty->GetIndexInArray() + 1;
+	if (NextIndex >= (int32)NumElements)
+	{
+		return 0;
+	}
+
+	TSharedPtr<IPropertyHandle> NextStructProperty = ParentArray->GetElement(NextIndex);
+	if (!NextStructProperty.IsValid())
+	{
+		return 0;
+	}
+	
+	TSharedPtr<IPropertyHandle> NextExpressionDepthProperty = NextStructProperty->GetChildHandle(TEXT("ExpressionDepth"));
+	if (!NextExpressionDepthProperty.IsValid())
+	{
+		return 0;
+	}
+	
+	uint8 Indent = 0;
+	NextExpressionDepthProperty->GetValue(Indent);
+
+	return Indent;
+}
+	
+FText FWorldConditionEditableDetails::GetOpenParens() const
+{
+	check(ExpressionDepthProperty);
+
+	const int32 CurrDepth = GetCurrExpressionDepth();
+	const int32 NextDepth = GetNextExpressionDepth();
+	const int32 DeltaDepth = NextDepth - CurrDepth;
+	const int32 OpenParens = FMath::Max(0, DeltaDepth);
+
+	FText Result = FText::GetEmpty();
+	
+	static_assert(UE::WorldCondition::MaxExpressionDepth == 4);
+	switch (OpenParens)
+	{
+	case 0:
+		break;
+	case 1:
+		Result = FText::FromString(TEXT("("));
+		break;
+	case 2:
+		Result = FText::FromString(TEXT("(("));
+		break;
+	case 3:
+		Result = FText::FromString(TEXT("((("));
+		break;
+	case 4:
+		Result = FText::FromString(TEXT("(((("));
+		break;
+	default:
+		ensureMsgf(false, TEXT("Max %d parenthesis supported, got %d."), UE::WorldCondition::MaxExpressionDepth, OpenParens);
+	}
+
+	return Result;
+}
+
+FText FWorldConditionEditableDetails::GetCloseParens() const
+{
+	check(ExpressionDepthProperty);
+
+	const int32 CurrDepth = GetCurrExpressionDepth();
+	const int32 NextDepth = GetNextExpressionDepth();
+	const int32 DeltaDepth = NextDepth - CurrDepth;
+	const int32 CloseParens = FMath::Max(0, -DeltaDepth);
+
+	FText Result = FText::GetEmpty();
+	
+	static_assert(UE::WorldCondition::MaxExpressionDepth == 4);
+	switch (CloseParens)
+	{
+	case 0:
+		break;
+	case 1:
+		Result = FText::FromString(TEXT(")"));
+		break;
+	case 2:
+		Result = FText::FromString(TEXT("))"));
+		break;
+	case 3:
+		Result = FText::FromString(TEXT(")))"));
+		break;
+	case 4:
+		Result = FText::FromString(TEXT("))))"));
+		break;
+	default:
+		ensureMsgf(false, TEXT("Max %d parenthesis supported, got %d."), UE::WorldCondition::MaxExpressionDepth, CloseParens);
+	}
+
+	return Result;
+}
+
+FText FWorldConditionEditableDetails::GetDisplayValueString() const
+{
+	if (const FWorldConditionEditable* EditableCondition = GetCommonCondition())
+	{
+		if (const FWorldConditionBase* Condition = EditableCondition->Condition.GetPtr<FWorldConditionBase>())
+		{
+			return Condition->GetDescription();
+		}
+	}
+	return FText();
+}
+
+TSharedRef<SWidget> FWorldConditionEditableDetails::GeneratePicker()
+{
+	TSharedRef<FWorldConditionStructFilter> StructFilter = MakeShared<FWorldConditionStructFilter>();
+	StructFilter->BaseStruct = TBaseStructure<FWorldConditionBase>::Get();
+	StructFilter->Schema = Schema;
+
+	FStructViewerInitializationOptions Options;
+	Options.bShowNoneOption = false;
+	Options.StructFilter = StructFilter;
+	Options.NameTypeToDisplay = EStructViewerNameTypeToDisplay::DisplayName;
+	Options.DisplayMode = EStructViewerDisplayMode::ListView;
+	Options.bAllowViewOptions = true;
+
+	const FOnStructPicked OnPicked(FOnStructPicked::CreateRaw(this, &FWorldConditionEditableDetails::OnStructPicked));
+
+	return SNew(SBox)
+		.WidthOverride(280)
+		[
+			SNew(SVerticalBox)
+			+ SVerticalBox::Slot()
+			.AutoHeight()
+			.MaxHeight(500)
+			[
+				FModuleManager::LoadModuleChecked<FStructViewerModule>("StructViewer").CreateStructViewer(Options, OnPicked)
+			]
+		];
+}
+
+void FWorldConditionEditableDetails::OnStructPicked(const UScriptStruct* InStruct) const
+{
+	check(StructProperty);
+
+	TArray<void*> RawNodeData;
+	StructProperty->AccessRawData(RawNodeData);
+
+	GEditor->BeginTransaction(LOCTEXT("SelectCondition", "Select Condition"));
+
+	StructProperty->NotifyPreChange();
+
+	for (void* Data : RawNodeData)
+	{
+		if (FWorldConditionEditable* Condition = static_cast<FWorldConditionEditable*>(Data))
+		{
+			if (InStruct)
+			{
+				Condition->Condition.InitializeAs(InStruct);
+			}
+		}
+	}
+
+	StructProperty->NotifyPostChange(EPropertyChangeType::ValueSet);
+	StructProperty->NotifyFinishedChangingProperties();
+
+	GEditor->EndTransaction();
+
+	ComboButton->SetIsOpen(false);
+
+	if (PropUtils)
+	{
+		PropUtils->ForceRefresh();
+	}
+}
+
+const FWorldConditionEditable* FWorldConditionEditableDetails::GetCommonCondition() const
+{
+	check(StructProperty);
+
+	bool bMultipleValues = false;
+	TArray<void*> RawData;
+	StructProperty->AccessRawData(RawData);
+
+	const FWorldConditionEditable* CommonCondition = nullptr;
+	
+	for (void* Data : RawData)
+	{
+		if (const FWorldConditionEditable* Condition = static_cast<FWorldConditionEditable*>(Data))
+		{
+			if (!bMultipleValues && Condition)
+			{
+				CommonCondition = Condition;
+			}
+			else if (CommonCondition != Condition)
+			{
+				CommonCondition = nullptr;
+				bMultipleValues = true;
+			}
+		}
+	}
+
+	return CommonCondition;
+}
+
+
+#undef LOCTEXT_NAMESPACE
diff --git a/Engine/Plugins/Runtime/WorldConditions/Source/WorldConditionsEditor/Private/Customizations/WorldConditionEditableDetails.h b/Engine/Plugins/Runtime/WorldConditions/Source/WorldConditionsEditor/Private/Customizations/WorldConditionEditableDetails.h
new file mode 100644
--- /dev/null
+++ b/Engine/Plugins/Runtime/WorldConditions/Source/WorldConditionsEditor/Private/Customizations/WorldConditionEditableDetails.h
@@ -0,0 +1,73 @@
+﻿// Copyright Epic Games, Inc. All Rights Reserved.
+
+#pragma once
+
+#include "PropertyHandle.h"
+#include "IDetailCustomization.h"
+#include "IPropertyTypeCustomization.h"
+
+class IPropertyHandle;
+class IPropertyUtilities;
+class SWidget;
+
+enum class EWorldConditionOperator : uint8;
+struct FWorldConditionEditable;
+class UWorldConditionSchema;
+
+/**
+ * Type customization for FWorldConditionEditable.
+ */
+class FWorldConditionEditableDetails : public IPropertyTypeCustomization
+{
+public:
+	/** Makes a new instance of this detail layout class for a specific detail view requesting it */
+	static TSharedRef<IPropertyTypeCustomization> MakeInstance();
+
+private:
+	/** IPropertyTypeCustomization interface */
+	virtual void CustomizeHeader(TSharedRef<IPropertyHandle> StructPropertyHandle, FDetailWidgetRow& HeaderRow, IPropertyTypeCustomizationUtils& StructCustomizationUtils) override;
+	virtual void CustomizeChildren(TSharedRef<IPropertyHandle> StructPropertyHandle, IDetailChildrenBuilder& StructBuilder, IPropertyTypeCustomizationUtils& StructCustomizationUtils) override;
+
+	void CacheSchema();
+
+	bool ShouldResetToDefault(TSharedPtr<IPropertyHandle> PropertyHandle) const;
+	void ResetToDefault(TSharedPtr<IPropertyHandle> PropertyHandle);
+	
+	FOptionalSize GetIndentSize() const;
+	TSharedRef<SWidget> OnGetIndentContent() const;
+
+	int32 GetExpressionDepth() const;
+	void SetExpressionDepth(const int32 Depth) const;
+	bool IsExpressionDepth(const int32 Depth) const;
+
+	FText GetOperatorText() const;
+	FSlateColor GetOperatorColor() const;
+	TSharedRef<SWidget> OnGetOperatorContent() const;
+	bool IsOperatorEnabled() const;
+
+	void SetOperator(const EWorldConditionOperator Operator) const;
+	bool IsOperator(const EWorldConditionOperator Operator) const;
+
+	bool IsFirstItem() const;
+	int32 GetCurrExpressionDepth() const;
+	int32 GetNextExpressionDepth() const;
+	
+	FText GetOpenParens() const;
+	FText GetCloseParens() const;
+
+	FText GetDisplayValueString() const;
+
+	TSharedRef<SWidget> GeneratePicker();
+	void OnStructPicked(const UScriptStruct* InStruct) const;
+	
+	const FWorldConditionEditable* GetCommonCondition() const;
+
+	const UWorldConditionSchema* Schema = nullptr;
+	
+	TSharedPtr<class SComboButton> ComboButton;
+	TSharedPtr<IPropertyHandle> StructProperty;
+	TSharedPtr<IPropertyHandle> ConditionProperty;
+	TSharedPtr<IPropertyHandle> ExpressionDepthProperty;
+	TSharedPtr<IPropertyHandle> OperatorProperty;
+	IPropertyUtilities* PropUtils = nullptr;
+};
diff --git a/Engine/Plugins/Runtime/WorldConditions/Source/WorldConditionsEditor/Private/Customizations/WorldConditionQueryDefinitionDetails.cpp b/Engine/Plugins/Runtime/WorldConditions/Source/WorldConditionsEditor/Private/Customizations/WorldConditionQueryDefinitionDetails.cpp
new file mode 100644
--- /dev/null
+++ b/Engine/Plugins/Runtime/WorldConditions/Source/WorldConditionsEditor/Private/Customizations/WorldConditionQueryDefinitionDetails.cpp
@@ -0,0 +1,70 @@
+﻿// Copyright Epic Games, Inc. All Rights Reserved.
+
+#include "WorldConditionQueryDefinitionDetails.h"
+#include "PropertyHandle.h"
+#include "IDetailChildrenBuilder.h"
+#include "DetailWidgetRow.h"
+#include "PropertyCustomizationHelpers.h"
+#include "DetailLayoutBuilder.h"
+#include "IPropertyUtilities.h"
+#include "ScopedTransaction.h"
+#include "Framework/MultiBox/MultiBoxBuilder.h"
+#include "WorldConditionQuery.h"
+
+#define LOCTEXT_NAMESPACE "WorldCondition"
+
+TSharedRef<IPropertyTypeCustomization> FWorldConditionQueryDefinitionDetails::MakeInstance()
+{
+	return MakeShareable(new FWorldConditionQueryDefinitionDetails);
+}
+
+void FWorldConditionQueryDefinitionDetails::CustomizeHeader(TSharedRef<IPropertyHandle> StructPropertyHandle, FDetailWidgetRow& HeaderRow, IPropertyTypeCustomizationUtils& StructCustomizationUtils)
+{
+	StructProperty = StructPropertyHandle;
+	PropUtils = StructCustomizationUtils.GetPropertyUtilities().Get();
+
+	EditableConditionsProperty = StructProperty->GetChildHandle(TEXT("EditableConditions"));
+
+	// Keep the definition up to date as it's being edited.
+	StructProperty->SetOnPropertyValueChanged(FSimpleDelegate::CreateLambda(
+	[StructProperty = StructProperty]()
+	{
+		TArray<void*> RawNodeData;
+		StructProperty->AccessRawData(RawNodeData);
+		for (void* Data : RawNodeData)
+		{
+			if (FWorldConditionQueryDefinition* QueryDefinition = static_cast<FWorldConditionQueryDefinition*>(Data))
+			{
+				QueryDefinition->Initialize();
+			}
+		}
+	}));
+	
+	HeaderRow
+		.NameContent()
+		[
+			StructPropertyHandle->CreatePropertyNameWidget()
+		]
+		.ValueContent()
+		[
+			EditableConditionsProperty->CreatePropertyValueWidget()
+		];
+}
+
+void FWorldConditionQueryDefinitionDetails::CustomizeChildren(TSharedRef<IPropertyHandle> StructPropertyHandle, IDetailChildrenBuilder& StructBuilder, IPropertyTypeCustomizationUtils& StructCustomizationUtils)
+{
+	check(EditableConditionsProperty.IsValid());
+
+	// Place editable conditions directly under.
+	if (EditableConditionsProperty.IsValid())
+	{
+		const TSharedRef<FDetailArrayBuilder> Builder = MakeShareable(new FDetailArrayBuilder(EditableConditionsProperty.ToSharedRef(), /*InGenerateHeader*/ false, /*InDisplayResetToDefault*/ true, /*InDisplayElementNum*/ false));
+		Builder->OnGenerateArrayElementWidget(FOnGenerateArrayElementWidget::CreateLambda([](TSharedRef<IPropertyHandle> PropertyHandle, int32 ArrayIndex, IDetailChildrenBuilder& ChildrenBuilder)
+		{
+			ChildrenBuilder.AddProperty(PropertyHandle);
+		}));
+		StructBuilder.AddCustomBuilder(Builder);
+	}
+}
+
+#undef LOCTEXT_NAMESPACE
diff --git a/Engine/Plugins/Runtime/WorldConditions/Source/WorldConditionsEditor/Private/Customizations/WorldConditionQueryDefinitionDetails.h b/Engine/Plugins/Runtime/WorldConditions/Source/WorldConditionsEditor/Private/Customizations/WorldConditionQueryDefinitionDetails.h
new file mode 100644
--- /dev/null
+++ b/Engine/Plugins/Runtime/WorldConditions/Source/WorldConditionsEditor/Private/Customizations/WorldConditionQueryDefinitionDetails.h
@@ -0,0 +1,31 @@
+﻿// Copyright Epic Games, Inc. All Rights Reserved.
+
+#pragma once
+
+#include "PropertyHandle.h"
+#include "IDetailCustomization.h"
+#include "IPropertyTypeCustomization.h"
+
+class IPropertyHandle;
+class IPropertyUtilities;
+class USmartObjectDefinition;
+class SWidget;
+
+/**
+ * Type customization for FWorldConditionQueryDefinition.
+ */
+class FWorldConditionQueryDefinitionDetails : public IPropertyTypeCustomization
+{
+public:
+	/** Makes a new instance of this detail layout class for a specific detail view requesting it */
+	static TSharedRef<IPropertyTypeCustomization> MakeInstance();
+
+private:
+	/** IPropertyTypeCustomization interface */
+	virtual void CustomizeHeader(TSharedRef<IPropertyHandle> StructPropertyHandle, FDetailWidgetRow& HeaderRow, IPropertyTypeCustomizationUtils& StructCustomizationUtils) override;
+	virtual void CustomizeChildren(TSharedRef<IPropertyHandle> StructPropertyHandle, IDetailChildrenBuilder& StructBuilder, IPropertyTypeCustomizationUtils& StructCustomizationUtils) override;
+
+	IPropertyUtilities* PropUtils = nullptr;
+	TSharedPtr<IPropertyHandle> StructProperty;
+	TSharedPtr<IPropertyHandle> EditableConditionsProperty;
+};
diff --git a/Engine/Plugins/Runtime/WorldConditions/Source/WorldConditionsEditor/Private/WorldConditionEditorStyle.cpp b/Engine/Plugins/Runtime/WorldConditions/Source/WorldConditionsEditor/Private/WorldConditionEditorStyle.cpp
new file mode 100644
--- /dev/null
+++ b/Engine/Plugins/Runtime/WorldConditions/Source/WorldConditionsEditor/Private/WorldConditionEditorStyle.cpp
@@ -0,0 +1,72 @@
+// Copyright Epic Games, Inc. All Rights Reserved.
+
+#include "WorldConditionEditorStyle.h"
+#include "Styling/SlateStyleRegistry.h"
+#include "Styling/StyleColors.h"
+#include "Styling/AppStyle.h"
+
+TUniquePtr<FWorldConditionEditorStyle> FWorldConditionEditorStyle::Instance(nullptr);
+FColor FWorldConditionEditorStyle::TypeColor(104,49,178);
+
+FWorldConditionEditorStyle::FWorldConditionEditorStyle() : FSlateStyleSet("WorldConditionEditorStyle")
+{
+	const FTextBlockStyle& NormalText = FAppStyle::Get().GetWidgetStyle<FTextBlockStyle>("NormalText");
+	const FComboButtonStyle& ComboButtonStyle = FCoreStyle::Get().GetWidgetStyle<FComboButtonStyle>("ComboButton");
+
+	// Condition Operator combo button
+	const FButtonStyle OperatorButton = FButtonStyle()
+		.SetNormal(FSlateRoundedBoxBrush(FStyleColors::AccentGreen.GetSpecifiedColor().Desaturate(0.3f), 4.0f))
+		.SetHovered(FSlateRoundedBoxBrush(FStyleColors::AccentGreen.GetSpecifiedColor().Desaturate(0.2f), 4.0f))
+		.SetPressed(FSlateRoundedBoxBrush(FStyleColors::AccentGreen.GetSpecifiedColor().Desaturate(0.1f), 4.0f))
+		.SetNormalForeground(FStyleColors::Foreground)
+		.SetHoveredForeground(FStyleColors::ForegroundHover)
+		.SetPressedForeground(FStyleColors::ForegroundHover)
+		.SetDisabledForeground(FStyleColors::ForegroundHover)
+		.SetNormalPadding(FMargin(2, 2, 2, 2))
+		.SetPressedPadding(FMargin(2, 3, 2, 1));
+
+	Set("Condition.Operator.ComboBox", FComboButtonStyle(ComboButtonStyle).SetButtonStyle(OperatorButton));
+
+	Set("Condition.Operator", FTextBlockStyle(NormalText)
+		.SetFont(FAppStyle::GetFontStyle(TEXT("PropertyWindow.BoldFont")))
+		.SetFontSize(7));
+
+	Set("Condition.Parens", FTextBlockStyle(NormalText)
+		.SetFont(FAppStyle::GetFontStyle(TEXT("PropertyWindow.NormalFont")))
+		.SetFontSize(12));
+
+	// Condition Indent combo button
+	const FButtonStyle IndentButton = FButtonStyle()
+		.SetNormal(FSlateRoundedBoxBrush(FLinearColor::Transparent, 2.0f))
+		.SetHovered(FSlateRoundedBoxBrush(FStyleColors::Background, 2.0f, FStyleColors::InputOutline, 1.0f))
+		.SetPressed(FSlateRoundedBoxBrush(FStyleColors::Background, 2.0f, FStyleColors::Hover, 1.0f))
+		.SetNormalForeground(FStyleColors::Transparent)
+		.SetHoveredForeground(FStyleColors::Hover)
+		.SetPressedForeground(FStyleColors::Foreground)
+		.SetNormalPadding(FMargin(2, 2, 2, 2))
+		.SetPressedPadding(FMargin(2, 3, 2, 1));
+	
+	Set("Condition.Indent.ComboBox", FComboButtonStyle(ComboButtonStyle).SetButtonStyle(IndentButton));
+
+	
+	FSlateStyleRegistry::RegisterSlateStyle(*this);
+}
+
+FWorldConditionEditorStyle::~FWorldConditionEditorStyle()
+{
+	FSlateStyleRegistry::UnRegisterSlateStyle(*this);
+}
+
+FWorldConditionEditorStyle& FWorldConditionEditorStyle::Get()
+{
+	if (!Instance.IsValid())
+	{
+		Instance = TUniquePtr<FWorldConditionEditorStyle>(new FWorldConditionEditorStyle);
+	}
+	return *(Instance.Get());
+}
+
+void FWorldConditionEditorStyle::Shutdown()
+{
+	Instance.Release();
+}
diff --git a/Engine/Plugins/Runtime/WorldConditions/Source/WorldConditionsEditor/Private/WorldConditionEditorStyle.h b/Engine/Plugins/Runtime/WorldConditions/Source/WorldConditionsEditor/Private/WorldConditionEditorStyle.h
new file mode 100644
--- /dev/null
+++ b/Engine/Plugins/Runtime/WorldConditions/Source/WorldConditionsEditor/Private/WorldConditionEditorStyle.h
@@ -0,0 +1,21 @@
+﻿// Copyright Epic Games, Inc. All Rights Reserved.
+
+#pragma once
+
+#include "Styling/SlateStyle.h"
+#include "Templates/UniquePtr.h"
+
+class FWorldConditionEditorStyle final : public FSlateStyleSet
+{
+public:
+	virtual ~FWorldConditionEditorStyle() override;
+
+	static FWorldConditionEditorStyle& Get();
+	static void Shutdown();
+
+	static FColor TypeColor;
+private:
+	FWorldConditionEditorStyle();
+
+	static TUniquePtr<FWorldConditionEditorStyle> Instance;
+};
diff --git a/Engine/Plugins/Runtime/WorldConditions/Source/WorldConditionsEditor/Private/WorldConditionsEditorModule.cpp b/Engine/Plugins/Runtime/WorldConditions/Source/WorldConditionsEditor/Private/WorldConditionsEditorModule.cpp
new file mode 100644
--- /dev/null
+++ b/Engine/Plugins/Runtime/WorldConditions/Source/WorldConditionsEditor/Private/WorldConditionsEditorModule.cpp
@@ -0,0 +1,55 @@
+// Copyright Epic Games, Inc. All Rights Reserved.
+
+#include "WorldConditionsEditorModule.h"
+
+#include "UnrealEdGlobals.h"
+#include "Editor/UnrealEdEngine.h"
+#include "WorldConditionEditorStyle.h"
+#include "Customizations/WorldConditionEditableDetails.h"
+#include "Customizations/WorldConditionQueryDefinitionDetails.h"
+#include "Customizations/WorldConditionContextDataRefDetails.h"
+
+#define LOCTEXT_NAMESPACE "WorldConditions"
+
+class FWorldConditionsEditorModule : public IWorldConditionsEditorModule
+{
+protected:
+	virtual void StartupModule() override;
+	virtual void ShutdownModule() override;
+
+	void RegisterComponentVisualizer(FName ComponentClassName, TSharedPtr<FComponentVisualizer> Visualizer);
+
+private:
+	TArray<FName> RegisteredComponentClassNames;
+};
+
+void FWorldConditionsEditorModule::StartupModule()
+{
+	FWorldConditionEditorStyle::Get();
+	
+	// Register the details customizer
+	FPropertyEditorModule& PropertyModule = FModuleManager::LoadModuleChecked<FPropertyEditorModule>("PropertyEditor");
+	PropertyModule.RegisterCustomPropertyTypeLayout("WorldConditionEditable", FOnGetPropertyTypeCustomizationInstance::CreateStatic(&FWorldConditionEditableDetails::MakeInstance));
+	PropertyModule.RegisterCustomPropertyTypeLayout("WorldConditionQueryDefinition", FOnGetPropertyTypeCustomizationInstance::CreateStatic(&FWorldConditionQueryDefinitionDetails::MakeInstance));
+	PropertyModule.RegisterCustomPropertyTypeLayout("WorldConditionContextDataRef", FOnGetPropertyTypeCustomizationInstance::CreateStatic(&FWorldConditionContextDataRefDetails::MakeInstance));
+
+	PropertyModule.NotifyCustomizationModuleChanged();
+}
+
+void FWorldConditionsEditorModule::ShutdownModule()
+{
+	// Unregister the details customization
+	if (FModuleManager::Get().IsModuleLoaded("PropertyEditor"))
+	{
+		FPropertyEditorModule& PropertyModule = FModuleManager::LoadModuleChecked<FPropertyEditorModule>("PropertyEditor");
+		PropertyModule.UnregisterCustomPropertyTypeLayout("WorldConditionEditable");
+		PropertyModule.UnregisterCustomPropertyTypeLayout("WorldConditionQueryDefinition");
+		PropertyModule.UnregisterCustomPropertyTypeLayout("WorldConditionContextDataRef");
+	}
+	
+	FWorldConditionEditorStyle::Shutdown();
+}
+
+IMPLEMENT_MODULE(FWorldConditionsEditorModule, WorldConditionsEditor)
+
+#undef LOCTEXT_NAMESPACE
diff --git a/Engine/Plugins/Runtime/WorldConditions/Source/WorldConditionsEditor/Public/WorldConditionsEditorModule.h b/Engine/Plugins/Runtime/WorldConditions/Source/WorldConditionsEditor/Public/WorldConditionsEditorModule.h
new file mode 100644
--- /dev/null
+++ b/Engine/Plugins/Runtime/WorldConditions/Source/WorldConditionsEditor/Public/WorldConditionsEditorModule.h
@@ -0,0 +1,35 @@
+// Copyright Epic Games, Inc. All Rights Reserved.
+
+#pragma  once
+
+#include "Modules/ModuleManager.h"
+
+/**
+ * The public interface to this module
+ */
+class IWorldConditionsEditorModule : public IModuleInterface
+{
+
+public:
+
+	/**
+	 * Singleton-like access to this module's interface.  This is just for convenience!
+	 * Beware of calling this during the shutdown phase, though.  Your module might have been unloaded already.
+	 *
+	 * @return Returns singleton instance, loading the module on demand if needed
+	 */
+	static inline IWorldConditionsEditorModule& Get()
+	{
+		return FModuleManager::LoadModuleChecked<IWorldConditionsEditorModule>("WorldConditionsEditor");
+	}
+
+	/**
+	 * Checks to see if this module is loaded and ready.  It is only valid to call Get() if IsAvailable() returns true.
+	 *
+	 * @return True if the module is loaded and ready to use
+	 */
+	static inline bool IsAvailable()
+	{
+		return FModuleManager::Get().IsModuleLoaded("WorldConditionsEditor");
+	}
+};
diff --git a/Engine/Plugins/Runtime/WorldConditions/Source/WorldConditionsEditor/WorldConditionsEditor.Build.cs b/Engine/Plugins/Runtime/WorldConditions/Source/WorldConditionsEditor/WorldConditionsEditor.Build.cs
new file mode 100644
--- /dev/null
+++ b/Engine/Plugins/Runtime/WorldConditions/Source/WorldConditionsEditor/WorldConditionsEditor.Build.cs
@@ -0,0 +1,39 @@
+// Copyright Epic Games, Inc. All Rights Reserved.
+
+namespace UnrealBuildTool.Rules
+{
+	public class WorldConditionsEditor : ModuleRules
+	{
+		public WorldConditionsEditor(ReadOnlyTargetRules Target) : base(Target)
+		{
+			PublicIncludePaths.AddRange(
+				new string[] {
+				}
+			);
+
+			PublicDependencyModuleNames.AddRange(
+				new [] {
+					"Core",
+					"CoreUObject",
+					"Engine",
+					"InputCore",
+					"PropertyEditor",
+					"SlateCore",
+					"Slate",
+					"WorldConditions",
+					"UnrealEd",
+					"StructUtils",
+				}
+			);
+
+			PrivateDependencyModuleNames.AddRange(
+				new [] {
+					"RenderCore",
+					"ApplicationCore",
+					"StructUtilsEditor",
+				}
+			);
+		}
+
+	}
+}
diff --git a/Engine/Plugins/Runtime/WorldConditions/Source/WorldConditionsTestSuite/Private/WorldConditionsTest.cpp b/Engine/Plugins/Runtime/WorldConditions/Source/WorldConditionsTestSuite/Private/WorldConditionsTest.cpp
new file mode 100644
--- /dev/null
+++ b/Engine/Plugins/Runtime/WorldConditions/Source/WorldConditionsTestSuite/Private/WorldConditionsTest.cpp
@@ -0,0 +1,174 @@
+// Copyright Epic Games, Inc. All Rights Reserved.
+
+#include "CoreMinimal.h"
+#include "AITestsCommon.h"
+#include "Engine/World.h"
+#include "WorldConditionTestTypes.h"
+#include "WorldConditionQuery.h"
+#include "WorldConditionSchema.h"
+#include "WorldConditionContext.h"
+
+#define LOCTEXT_NAMESPACE "AITestSuite_WorldConditionsTest"
+
+PRAGMA_DISABLE_OPTIMIZATION
+
+struct FWorldConditionTest_Init : FAITestBase
+{
+	virtual bool InstantTest() override
+	{
+		FWorldConditionQueryDefinition Definition;
+		Definition.SchemaClass = UWorldConditionTestSchema::StaticClass();
+		Definition.EditableConditions.Emplace(0, EWorldConditionOperator::Copy, FConstStructView::Make(FWorldConditionTest()));
+		Definition.EditableConditions.Emplace(0, EWorldConditionOperator::And, FConstStructView::Make(FWorldConditionTest()));
+		const bool bInitialized = Definition.Initialize();
+
+		AITEST_TRUE("Query definition should get initialized", bInitialized);
+
+		FWorldConditionQueryState State;
+		const bool bStateInitialized = State.Initialize(GetWorld(), Definition);
+		AITEST_TRUE("Query state should get initialized", bStateInitialized);
+		AITEST_EQUAL("Query state should have 2 conditions", State.GetNumConditions(), 2);
+		AITEST_EQUAL("Query state condition 1 operator should be And", State.GetItem(1).Operator, EWorldConditionOperator::And);
+		
+		State.Free(Definition);
+		AITEST_EQUAL("Query state should have 0 conditions", State.GetNumConditions(), 0);
+		
+		return true;
+	}
+};
+IMPLEMENT_AI_INSTANT_TEST(FWorldConditionTest_Init, "System.WorldConditions.Init");
+
+struct FWorldConditionTest_Eval : FAITestBase
+{
+	virtual bool InstantTest() override
+	{
+		FWorldConditionQuery Query;
+		const bool bInitialized = Query.DebugInitialize(UWorldConditionTestSchema::StaticClass(),
+			{
+				FWorldConditionEditable(0, EWorldConditionOperator::Copy, FConstStructView::Make(FWorldConditionTest(1))),
+				FWorldConditionEditable(0, EWorldConditionOperator::And, FConstStructView::Make(FWorldConditionTest(1)))
+			});
+
+		AITEST_TRUE("Query should get initialized", bInitialized);
+		
+		const FWorldConditionTestData TestData(1);
+
+		const UWorldConditionTestSchema* TestSchema = GetDefault<UWorldConditionTestSchema>();
+		FWorldConditionContextData ContextData(*TestSchema);
+		const bool bContextDataSet = ContextData.SetContextData(TestSchema->GetValueRef(), &TestData);
+		AITEST_TRUE("Query context data should be set", bContextDataSet);
+		
+		const bool bActivated = Query.Activate(GetWorld(), ContextData);
+		AITEST_TRUE("Query should get activated", bActivated);
+		AITEST_TRUE("Query should be active", Query.IsActive());
+		
+		const bool bResult = Query.IsTrue(ContextData);
+		AITEST_TRUE("Query result should be true", bResult);
+
+		Query.Deactivate(ContextData);
+		AITEST_FALSE("Query should not be active", Query.IsActive());
+
+		return true;
+	}
+};
+IMPLEMENT_AI_INSTANT_TEST(FWorldConditionTest_Eval, "System.WorldConditions.Eval");
+
+struct FWorldConditionTest_CachedEval : FAITestBase
+{
+	virtual bool InstantTest() override
+	{
+		FWorldConditionQuery Query;
+		const bool bInitialized = Query.DebugInitialize(UWorldConditionTestCachedSchema::StaticClass(),
+			{
+				FWorldConditionEditable(0, EWorldConditionOperator::Copy, FConstStructView::Make(FWorldConditionTestCached(1))),
+				FWorldConditionEditable(0, EWorldConditionOperator::And, FConstStructView::Make(FWorldConditionTestCached(1)))
+			});
+
+		AITEST_TRUE("Query should get initialized", bInitialized);
+		
+		FWorldConditionTestData TestData(1);
+
+		const UWorldConditionTestCachedSchema* TestSchema = GetDefault<UWorldConditionTestCachedSchema>();
+		FWorldConditionContextData ContextData(*TestSchema);
+		const bool bContextDataSet = ContextData.SetContextData(TestSchema->GetValueRef(), &TestData);
+		AITEST_TRUE("Query context data should be set", bContextDataSet);
+		
+		const bool bActivated = Query.Activate(GetWorld(), ContextData);
+		AITEST_TRUE("Query should get activated", bActivated);
+		AITEST_TRUE("Query should be active", Query.IsActive());
+		
+		const bool bResult1 = Query.IsTrue(ContextData);
+		AITEST_TRUE("Query result should be true", bResult1);
+		AITEST_EQUAL("Test data should have been accessed 2 times", TestData.AccessCount, 2);
+
+		// This test should get cached.
+		const bool bResult2 = Query.IsTrue(ContextData);
+		AITEST_TRUE("Query result should be true", bResult2);
+		AITEST_EQUAL("Test data should have been accessed 2 times", TestData.AccessCount, 2);
+
+		// Change and invalidate cached
+		TestData.Value = 0;
+		if (TestData.ValueChanged.IsBound())
+		{
+			TestData.ValueChanged.Broadcast();
+		}
+		
+		// This test should get cached.
+		const bool bResult3 = Query.IsTrue(ContextData);
+		AITEST_FALSE("Query result should be false", bResult3);
+		AITEST_EQUAL("Test data should have been accessed 4 times", TestData.AccessCount, 4);
+		
+		Query.Deactivate(ContextData);
+		AITEST_FALSE("Query should not be active", Query.IsActive());
+
+		return true;
+	}
+};
+IMPLEMENT_AI_INSTANT_TEST(FWorldConditionTest_CachedEval, "System.WorldConditions.CachedEval");
+
+struct FWorldConditionTest_EvalComplex : FAITestBase
+{
+	virtual bool InstantTest() override
+	{
+		FWorldConditionQuery Query;
+		const bool bInitialized = Query.DebugInitialize(UWorldConditionTestSchema::StaticClass(),
+			{
+				FWorldConditionEditable(/*Depth*/0, EWorldConditionOperator::Copy, FConstStructView::Make(FWorldConditionTest(0))),	//	if	(A
+				FWorldConditionEditable(/*Depth*/1, EWorldConditionOperator::Or,   FConstStructView::Make(FWorldConditionTest(1))),	//	.	|| B)
+				FWorldConditionEditable(/*Depth*/0, EWorldConditionOperator::And,  FConstStructView::Make(FWorldConditionTest(1))),	//	&&	(	(C
+				FWorldConditionEditable(/*Depth*/2, EWorldConditionOperator::And,  FConstStructView::Make(FWorldConditionTest(1))),	//	.	.	&& D)
+				FWorldConditionEditable(/*Depth*/1, EWorldConditionOperator::Or,   FConstStructView::Make(FWorldConditionTest(0))),	//	.	|| E)
+			});
+
+		AITEST_TRUE("Query should get initialized", bInitialized);
+		
+		FWorldConditionTestData TestData(1);
+
+		const UWorldConditionTestSchema* TestSchema = GetDefault<UWorldConditionTestSchema>();
+		FWorldConditionContextData ContextData(*TestSchema);
+		const bool bContextDataSet = ContextData.SetContextData(TestSchema->GetValueRef(), &TestData);
+		AITEST_TRUE("Query context data should be set", bContextDataSet);
+		
+		const bool bActivated = Query.Activate(GetWorld(), ContextData);
+		AITEST_TRUE("Query should get activated", bActivated);
+		AITEST_TRUE("Query should be active", Query.IsActive());
+		
+		const bool bResult1 = Query.IsTrue(ContextData);
+		AITEST_TRUE("Query result1 should be true", bResult1);
+
+		TestData.Value = 0;
+
+		const bool bResult2 = Query.IsTrue(ContextData);
+		AITEST_TRUE("Query result2 should be true", bResult2);
+
+		Query.Deactivate(ContextData);
+		AITEST_FALSE("Query should not be active", Query.IsActive());
+
+		return true;
+	}
+};
+IMPLEMENT_AI_INSTANT_TEST(FWorldConditionTest_EvalComplex, "System.WorldConditions.EvalComplex");
+
+PRAGMA_ENABLE_OPTIMIZATION
+
+#undef LOCTEXT_NAMESPACE
diff --git a/Engine/Plugins/Runtime/WorldConditions/Source/WorldConditionsTestSuite/Private/WorldConditionsTestSuiteModule.cpp b/Engine/Plugins/Runtime/WorldConditions/Source/WorldConditionsTestSuite/Private/WorldConditionsTestSuiteModule.cpp
new file mode 100644
--- /dev/null
+++ b/Engine/Plugins/Runtime/WorldConditions/Source/WorldConditionsTestSuite/Private/WorldConditionsTestSuiteModule.cpp
@@ -0,0 +1,13 @@
+// Copyright Epic Games, Inc. All Rights Reserved.
+
+#include "WorldConditionsTestSuiteModule.h"
+
+#define LOCTEXT_NAMESPACE "WorldConditionsTestSuite"
+
+class FWorldConditionsTestSuiteModule : public IWorldConditionsTestSuiteModule
+{
+};
+
+IMPLEMENT_MODULE(FWorldConditionsTestSuiteModule, WorldConditionsTestSuite)
+
+#undef LOCTEXT_NAMESPACE
diff --git a/Engine/Plugins/Runtime/WorldConditions/Source/WorldConditionsTestSuite/Public/WorldConditionTestTypes.h b/Engine/Plugins/Runtime/WorldConditions/Source/WorldConditionsTestSuite/Public/WorldConditionTestTypes.h
new file mode 100644
--- /dev/null
+++ b/Engine/Plugins/Runtime/WorldConditions/Source/WorldConditionsTestSuite/Public/WorldConditionTestTypes.h
@@ -0,0 +1,210 @@
+﻿// Copyright Epic Games, Inc. All Rights Reserved.
+
+#pragma once
+
+#include "StructView.h"
+#include "WorldConditionTypes.h"
+#include "WorldConditionSchema.h"
+#include "WorldConditionBase.h"
+#include "WorldConditionContext.h"
+#include "WorldConditionTestTypes.generated.h"
+
+USTRUCT()
+struct FWorldConditionTestData
+{
+	GENERATED_BODY()
+
+	FWorldConditionTestData() = default;
+	explicit FWorldConditionTestData(const int32 InValue) : Value(InValue) {} 
+
+	FSimpleMulticastDelegate ValueChanged;
+
+	int32 Value = 0;
+	int32 AccessCount = 0;
+};
+
+UCLASS()
+class UWorldConditionTestSchema : public UWorldConditionSchema
+{
+	GENERATED_BODY()
+public:	
+	explicit UWorldConditionTestSchema(const FObjectInitializer& ObjectInitializer)
+		: Super(ObjectInitializer)
+	{
+		ActorRef = AddContextDataDesc(TEXT("Actor"), AActor::StaticClass(), EWorldConditionContextDataType::Persistent);
+		ValueRef = AddContextDataDesc(TEXT("Value"), FWorldConditionTestData::StaticStruct(), EWorldConditionContextDataType::Dynamic);
+	}
+
+	virtual bool IsStructAllowed(const UScriptStruct* InScriptStruct) const override
+	{
+		return InScriptStruct && InScriptStruct->IsChildOf(TBaseStructure<FWorldConditionBase>::Get());
+	}
+
+	FWorldConditionContextDataRef GetActorRef() const { return ActorRef; }
+	FWorldConditionContextDataRef GetValueRef() const { return ValueRef; };
+
+private:
+	
+	FWorldConditionContextDataRef ActorRef;
+	FWorldConditionContextDataRef ValueRef;
+};
+
+
+UCLASS()
+class UWorldConditionTestCachedSchema : public UWorldConditionSchema
+{
+	GENERATED_BODY()
+public:	
+	explicit UWorldConditionTestCachedSchema(const FObjectInitializer& ObjectInitializer)
+		: Super(ObjectInitializer)
+	{
+		ValueRef = AddContextDataDesc(TEXT("Value"), FWorldConditionTestData::StaticStruct(), EWorldConditionContextDataType::Persistent);
+	}
+
+	virtual bool IsStructAllowed(const UScriptStruct* InScriptStruct) const override
+	{
+		return InScriptStruct && InScriptStruct->IsChildOf(TBaseStructure<FWorldConditionBase>::Get());
+	}
+
+	FWorldConditionContextDataRef GetValueRef() const { return ValueRef; };
+
+private:
+	
+	FWorldConditionContextDataRef ValueRef;
+};
+
+USTRUCT(meta=(Hidden))
+struct FWorldConditionTest : public FWorldConditionBase
+{
+	GENERATED_BODY()
+
+	FWorldConditionTest() = default;
+	explicit FWorldConditionTest(const int32 InTestValue) : TestValue(InTestValue) {}
+	
+	virtual const UStruct* GetRuntimeStateType() const override { return nullptr; }
+	
+	virtual bool Initialize(const UWorldConditionSchema& Schema) override
+	{
+		const UWorldConditionTestSchema* TestSchema = Cast<UWorldConditionTestSchema>(&Schema);
+		if (TestSchema == nullptr)
+		{
+			return false;
+		}
+		ValueRef = TestSchema->GetValueRef();
+		bCanCacheResult = false;
+		return true;
+	}
+	
+	virtual bool Activate(const FWorldConditionContext& Context) const override
+	{
+		return true;
+	}
+	
+	virtual EWorldConditionResult IsTrue(const FWorldConditionContext& Context) const override
+	{
+		if (const FWorldConditionTestData* TestData = Context.GetContextDataPtr<FWorldConditionTestData>(ValueRef))
+		{
+			return (TestData->Value == TestValue) ? EWorldConditionResult::IsTrue : EWorldConditionResult::IsFalse;
+		}
+		return EWorldConditionResult::IsFalse;
+	}
+	
+	virtual void Deactivate(const FWorldConditionContext& Context) const override
+	{
+		
+	}
+
+protected:
+
+	FWorldConditionContextDataRef ValueRef;
+	
+	int32 TestValue = 0;
+};
+
+
+USTRUCT()
+struct FWorldConditionTestState
+{
+	GENERATED_BODY()
+
+	FDelegateHandle DelegateHandle;
+};
+
+USTRUCT(meta=(Hidden))
+struct FWorldConditionTestCached : public FWorldConditionBase
+{
+	GENERATED_BODY()
+
+	FWorldConditionTestCached() = default;
+	explicit FWorldConditionTestCached(const int32 InTestValue) : TestValue(InTestValue) {}
+	
+	using FStateType = FWorldConditionTestState;
+	
+	virtual const UStruct* GetRuntimeStateType() const override { return FStateType::StaticStruct(); }
+	
+	virtual bool Initialize(const UWorldConditionSchema& Schema) override
+	{
+		const UWorldConditionTestCachedSchema* TestSchema = Cast<UWorldConditionTestCachedSchema>(&Schema);
+		if (TestSchema == nullptr)
+		{
+			return false;
+		}
+		ValueRef = TestSchema->GetValueRef();
+
+		bCanCacheResult = Schema.GetContextDataTypeByRef(ValueRef) == EWorldConditionContextDataType::Persistent;
+		
+		return true;
+	}
+	
+	virtual bool Activate(const FWorldConditionContext& Context) const override
+	{
+		if (Context.GetContextDataType(ValueRef) == EWorldConditionContextDataType::Persistent)
+		{
+			if (const FWorldConditionTestData* TestData = Context.GetContextDataPtr<FWorldConditionTestData>(ValueRef))
+			{
+				FStateType& State = Context.GetState(*this);
+				FWorldConditionTestData* MutableTestData = const_cast<FWorldConditionTestData*>(TestData);
+				State.DelegateHandle = MutableTestData->ValueChanged.AddLambda([this, &QueryState = Context.GetQueryState()]()
+					{
+						InvalidateResult(QueryState);
+					});
+				
+				return true;
+			}
+			return false;
+		}
+		
+		return true;
+	}
+	
+	virtual EWorldConditionResult IsTrue(const FWorldConditionContext& Context) const override
+	{
+		if (const FWorldConditionTestData* TestData = Context.GetContextDataPtr<FWorldConditionTestData>(ValueRef))
+		{
+			FWorldConditionTestData* MutableTestData = const_cast<FWorldConditionTestData*>(TestData);
+			MutableTestData->AccessCount++;
+			return (TestData->Value == TestValue) ? EWorldConditionResult::IsTrue : EWorldConditionResult::IsFalse;
+		}
+		return EWorldConditionResult::IsFalse;
+	}
+	
+	virtual void Deactivate(const FWorldConditionContext& Context) const override
+	{
+		FStateType& State = Context.GetState(*this);
+		if (State.DelegateHandle.IsValid())
+		{
+			if (const FWorldConditionTestData* TestData = Context.GetContextDataPtr<FWorldConditionTestData>(ValueRef))
+			{
+				FWorldConditionTestData* MutableTestData = const_cast<FWorldConditionTestData*>(TestData);
+				MutableTestData->ValueChanged.Remove(State.DelegateHandle);
+				State.DelegateHandle.Reset();
+			}
+		}
+	}
+
+protected:
+
+	FWorldConditionContextDataRef ValueRef;
+	
+	int32 TestValue = 0;
+};
\ No newline at end of file
diff --git a/Engine/Plugins/Runtime/WorldConditions/Source/WorldConditionsTestSuite/Public/WorldConditionsTestSuiteModule.h b/Engine/Plugins/Runtime/WorldConditions/Source/WorldConditionsTestSuite/Public/WorldConditionsTestSuiteModule.h
new file mode 100644
--- /dev/null
+++ b/Engine/Plugins/Runtime/WorldConditions/Source/WorldConditionsTestSuite/Public/WorldConditionsTestSuiteModule.h
@@ -0,0 +1,37 @@
+// Copyright Epic Games, Inc. All Rights Reserved.
+
+#pragma  once
+
+#include "CoreMinimal.h"
+#include "Modules/ModuleManager.h"
+
+/**
+* The public interface to this module
+*/
+class IWorldConditionsTestSuiteModule : public IModuleInterface
+{
+
+public:
+
+	/**
+	* Singleton-like access to this module's interface.  This is just for convenience!
+	* Beware of calling this during the shutdown phase, though.  Your module might have been unloaded already.
+	*
+	* @return Returns singleton instance, loading the module on demand if needed
+	*/
+	static inline IWorldConditionsTestSuiteModule& Get()
+	{
+		return FModuleManager::LoadModuleChecked<IWorldConditionsTestSuiteModule>("WorldConditionsTestSuite");
+	}
+
+	/**
+	* Checks to see if this module is loaded and ready.  It is only valid to call Get() if IsAvailable() returns true.
+	*
+	* @return True if the module is loaded and ready to use
+	*/
+	static inline bool IsAvailable()
+	{
+		return FModuleManager::Get().IsModuleLoaded("WorldConditionsTestSuite");
+	}
+};
+
diff --git a/Engine/Plugins/Runtime/WorldConditions/Source/WorldConditionsTestSuite/WorldConditionsTestSuite.Build.cs b/Engine/Plugins/Runtime/WorldConditions/Source/WorldConditionsTestSuite/WorldConditionsTestSuite.Build.cs
new file mode 100644
--- /dev/null
+++ b/Engine/Plugins/Runtime/WorldConditions/Source/WorldConditionsTestSuite/WorldConditionsTestSuite.Build.cs
@@ -0,0 +1,33 @@
+// Copyright Epic Games, Inc. All Rights Reserved.
+
+namespace UnrealBuildTool.Rules
+{
+	public class WorldConditionsTestSuite : ModuleRules
+	{
+		public WorldConditionsTestSuite(ReadOnlyTargetRules Target) : base(Target)
+		{
+			PublicIncludePaths.AddRange(
+				new string[] {
+				}
+			);
+
+			PublicDependencyModuleNames.AddRange(
+				new[] {
+					"AITestSuite",
+					"Core",
+					"CoreUObject",
+					"Engine",
+					"GameplayTags",
+					"WorldConditions",
+					"StructUtils",
+				}
+			);
+
+			if (Target.bBuildEditor)
+			{
+				PrivateDependencyModuleNames.Add("EditorFramework");
+				PrivateDependencyModuleNames.Add("UnrealEd");
+			}
+		}
+	}
+}
\ No newline at end of file
diff --git a/Engine/Plugins/Runtime/WorldConditions/WorldConditions.uplugin b/Engine/Plugins/Runtime/WorldConditions/WorldConditions.uplugin
new file mode 100644
--- /dev/null
+++ b/Engine/Plugins/Runtime/WorldConditions/WorldConditions.uplugin
@@ -0,0 +1,42 @@
+{
+	"FileVersion" : 3,
+	"Version" : 1,
+	"VersionName" : "0.1",
+	"FriendlyName" : "World Conditions",
+	"Description" : "General purpose cached conditions",
+	"Category" : "Gameplay",
+	"CreatedBy" : "Epic Games, Inc.",
+	"CreatedByURL" : "http://epicgames.com",
+	"DocsURL" : "",
+	"MarketplaceURL" : "",
+	"SupportURL" : "",
+	"EnabledByDefault" : false,
+	"CanContainContent" : true,
+	"IsExperimentalVersion": true,
+	"Installed" : false,
+	"Modules":
+	[
+		{
+			"Name": "WorldConditions",
+			"Type": "Runtime",
+			"LoadingPhase": "Default"
+		},
+		{
+			"Name" : "WorldConditionsEditor",
+			"Type" : "Editor",
+			"LoadingPhase" : "PostEngineInit"
+		},
+		{
+			"Name": "WorldConditionsTestSuite",
+			"Type": "UncookedOnly",
+			"LoadingPhase": "Default"
+		}
+	],
+	"Plugins":
+	[
+		{
+			"Name": "StructUtils",
+			"Enabled": true
+		}
+	]
+}
\ No newline at end of file
