World Condition: Fix missing headers causing non-unity build errors.

#br trivial

diff --git a/Engine/Plugins/Runtime/WorldConditions/Source/WorldConditions/Public/WorldConditionQuery.h b/Engine/Plugins/Runtime/WorldConditions/Source/WorldConditions/Public/WorldConditionQuery.h
--- a/Engine/Plugins/Runtime/WorldConditions/Source/WorldConditions/Public/WorldConditionQuery.h
+++ b/Engine/Plugins/Runtime/WorldConditions/Source/WorldConditions/Public/WorldConditionQuery.h
@@ -7,6 +7,7 @@
 #include "StructView.h"
 #include "WorldConditionBase.h"
 #include "WorldConditionTypes.h"
+#include "Templates/SubclassOf.h"
 #include "WorldConditionQuery.generated.h"
 
 struct FWorldConditionContext;
diff --git a/Engine/Plugins/Runtime/WorldConditions/Source/WorldConditionsEditor/Private/WorldConditionEditorStyle.cpp b/Engine/Plugins/Runtime/WorldConditions/Source/WorldConditionsEditor/Private/WorldConditionEditorStyle.cpp
--- a/Engine/Plugins/Runtime/WorldConditions/Source/WorldConditionsEditor/Private/WorldConditionEditorStyle.cpp
+++ b/Engine/Plugins/Runtime/WorldConditions/Source/WorldConditionsEditor/Private/WorldConditionEditorStyle.cpp
@@ -1,6 +1,7 @@
 // Copyright Epic Games, Inc. All Rights Reserved.
 
 #include "WorldConditionEditorStyle.h"
+#include "Styling/SlateTypes.h"
 #include "Styling/SlateStyleRegistry.h"
 #include "Styling/StyleColors.h"
 #include "Styling/AppStyle.h"
diff --git a/Engine/Plugins/Runtime/WorldConditions/Source/WorldConditionsTestSuite/Public/WorldConditionTestTypes.h b/Engine/Plugins/Runtime/WorldConditions/Source/WorldConditionsTestSuite/Public/WorldConditionTestTypes.h
--- a/Engine/Plugins/Runtime/WorldConditions/Source/WorldConditionsTestSuite/Public/WorldConditionTestTypes.h
+++ b/Engine/Plugins/Runtime/WorldConditions/Source/WorldConditionsTestSuite/Public/WorldConditionTestTypes.h
@@ -7,6 +7,7 @@
 #include "WorldConditionSchema.h"
 #include "WorldConditionBase.h"
 #include "WorldConditionContext.h"
+#include "GameFramework/Actor.h"
 #include "WorldConditionTestTypes.generated.h"
 
 USTRUCT()
