World Condition: More fixes for CI errors.

diff --git a/Engine/Plugins/Runtime/WorldConditions/Source/WorldConditionsEditor/Private/Customizations/WorldConditionEditableDetails.cpp b/Engine/Plugins/Runtime/WorldConditions/Source/WorldConditionsEditor/Private/Customizations/WorldConditionEditableDetails.cpp
--- a/Engine/Plugins/Runtime/WorldConditions/Source/WorldConditionsEditor/Private/Customizations/WorldConditionEditableDetails.cpp
+++ b/Engine/Plugins/Runtime/WorldConditions/Source/WorldConditionsEditor/Private/Customizations/WorldConditionEditableDetails.cpp
@@ -17,6 +17,8 @@
 #include "WorldConditionEditorStyle.h"
 #include "InstancedStructDetails.h"
 #include "WorldConditionSchema.h"
+#include "Editor.h"
+#include "Modules/ModuleManager.h"
 
 #define LOCTEXT_NAMESPACE "WorldCondition"
 
