CIS unity fixups

diff --git a/Engine/Plugins/Runtime/WorldConditions/Source/WorldConditions/Private/WorldConditionBase.cpp b/Engine/Plugins/Runtime/WorldConditions/Source/WorldConditions/Private/WorldConditionBase.cpp
--- a/Engine/Plugins/Runtime/WorldConditions/Source/WorldConditions/Private/WorldConditionBase.cpp
+++ b/Engine/Plugins/Runtime/WorldConditions/Source/WorldConditions/Private/WorldConditionBase.cpp
@@ -1,7 +1,8 @@
-﻿// Copyright Epic Games, Inc. All Rights Reserved.
+// Copyright Epic Games, Inc. All Rights Reserved.
 
 #include "WorldConditionBase.h"
 #include "WorldConditionQuery.h"
+#include "WorldConditionSchema.h"
 
 FWorldConditionBase::~FWorldConditionBase()
 {
diff --git a/Engine/Plugins/Runtime/WorldConditions/Source/WorldConditions/Private/WorldConditionQuery.cpp b/Engine/Plugins/Runtime/WorldConditions/Source/WorldConditions/Private/WorldConditionQuery.cpp
--- a/Engine/Plugins/Runtime/WorldConditions/Source/WorldConditions/Private/WorldConditionQuery.cpp
+++ b/Engine/Plugins/Runtime/WorldConditions/Source/WorldConditions/Private/WorldConditionQuery.cpp
@@ -2,6 +2,7 @@
 
 #include "WorldConditionQuery.h"
 #include "WorldConditionContext.h"
+#include "WorldConditionSchema.h"
 
 FWorldConditionQueryState::~FWorldConditionQueryState()
 {
diff --git a/Engine/Plugins/Runtime/WorldConditions/Source/WorldConditions/Public/WorldConditionQuery.h b/Engine/Plugins/Runtime/WorldConditions/Source/WorldConditions/Public/WorldConditionQuery.h
--- a/Engine/Plugins/Runtime/WorldConditions/Source/WorldConditions/Public/WorldConditionQuery.h
+++ b/Engine/Plugins/Runtime/WorldConditions/Source/WorldConditions/Public/WorldConditionQuery.h
@@ -12,6 +12,7 @@
 
 struct FWorldConditionContext;
 struct FWorldConditionContextData;
+class UWorldConditionSchema;
 
 /**
  * World Condition Query is an expression of World Conditions whose state can be queried.
diff --git a/Engine/Plugins/Runtime/WorldConditions/Source/WorldConditionsEditor/Private/Customizations/WorldConditionContextDataRefDetails.cpp b/Engine/Plugins/Runtime/WorldConditions/Source/WorldConditionsEditor/Private/Customizations/WorldConditionContextDataRefDetails.cpp
--- a/Engine/Plugins/Runtime/WorldConditions/Source/WorldConditionsEditor/Private/Customizations/WorldConditionContextDataRefDetails.cpp
+++ b/Engine/Plugins/Runtime/WorldConditions/Source/WorldConditionsEditor/Private/Customizations/WorldConditionContextDataRefDetails.cpp
@@ -14,6 +14,7 @@
 #include "WorldConditionTypes.h"
 #include "WorldConditionQuery.h"
 #include "WorldConditionSchema.h"
+#include "Editor.h"
 
 #define LOCTEXT_NAMESPACE "WorldCondition"
 
diff --git a/Engine/Plugins/Runtime/WorldConditions/Source/WorldConditionsEditor/Private/WorldConditionEditorStyle.cpp b/Engine/Plugins/Runtime/WorldConditions/Source/WorldConditionsEditor/Private/WorldConditionEditorStyle.cpp
--- a/Engine/Plugins/Runtime/WorldConditions/Source/WorldConditionsEditor/Private/WorldConditionEditorStyle.cpp
+++ b/Engine/Plugins/Runtime/WorldConditions/Source/WorldConditionsEditor/Private/WorldConditionEditorStyle.cpp
@@ -5,6 +5,7 @@
 #include "Styling/SlateStyleRegistry.h"
 #include "Styling/StyleColors.h"
 #include "Styling/AppStyle.h"
+#include "Styling/CoreStyle.h"
 
 TUniquePtr<FWorldConditionEditorStyle> FWorldConditionEditorStyle::Instance(nullptr);
 FColor FWorldConditionEditorStyle::TypeColor(104,49,178);
