Change WorldConditionQuery to ANSI text instead of utf8, fix issue where TSubClassOf needs the include in the header instead of forward declared

diff --git a/Engine/Plugins/Runtime/WorldConditions/Source/WorldConditions/Private/WorldConditionQuery.cpp b/Engine/Plugins/Runtime/WorldConditions/Source/WorldConditions/Private/WorldConditionQuery.cpp
--- a/Engine/Plugins/Runtime/WorldConditions/Source/WorldConditions/Private/WorldConditionQuery.cpp
+++ b/Engine/Plugins/Runtime/WorldConditions/Source/WorldConditions/Private/WorldConditionQuery.cpp
@@ -1,8 +1,7 @@
-﻿// Copyright Epic Games, Inc. All Rights Reserved.
+// Copyright Epic Games, Inc. All Rights Reserved.
 
 #include "WorldConditionQuery.h"
 #include "WorldConditionContext.h"
-#include "WorldConditionSchema.h"
 
 FWorldConditionQueryState::~FWorldConditionQueryState()
 {
diff --git a/Engine/Plugins/Runtime/WorldConditions/Source/WorldConditions/Public/WorldConditionQuery.h b/Engine/Plugins/Runtime/WorldConditions/Source/WorldConditions/Public/WorldConditionQuery.h
--- a/Engine/Plugins/Runtime/WorldConditions/Source/WorldConditions/Public/WorldConditionQuery.h
+++ b/Engine/Plugins/Runtime/WorldConditions/Source/WorldConditions/Public/WorldConditionQuery.h
@@ -1,4 +1,4 @@
-﻿// Copyright Epic Games, Inc. All Rights Reserved.
+// Copyright Epic Games, Inc. All Rights Reserved.
 
 #pragma once
 
@@ -8,11 +8,11 @@
 #include "WorldConditionBase.h"
 #include "WorldConditionTypes.h"
 #include "Templates/SubclassOf.h"
+#include "WorldConditionSchema.h"
 #include "WorldConditionQuery.generated.h"
 
 struct FWorldConditionContext;
 struct FWorldConditionContextData;
-class UWorldConditionSchema;
 
 /**
  * World Condition Query is an expression of World Conditions whose state can be queried.
