remove deprecated PRAGMA_DISABLE_OPTIMIZATION

diff --git a/Engine/Plugins/Runtime/WorldConditions/Source/WorldConditionsTestSuite/Private/WorldConditionsTest.cpp b/Engine/Plugins/Runtime/WorldConditions/Source/WorldConditionsTestSuite/Private/WorldConditionsTest.cpp
--- a/Engine/Plugins/Runtime/WorldConditions/Source/WorldConditionsTestSuite/Private/WorldConditionsTest.cpp
+++ b/Engine/Plugins/Runtime/WorldConditions/Source/WorldConditionsTestSuite/Private/WorldConditionsTest.cpp
@@ -10,8 +10,6 @@
 
 #define LOCTEXT_NAMESPACE "AITestSuite_WorldConditionsTest"
 
-PRAGMA_DISABLE_OPTIMIZATION
-
 struct FWorldConditionTest_Init : FAITestBase
 {
 	virtual bool InstantTest() override
@@ -169,6 +167,4 @@ struct FWorldConditionTest_EvalComplex : FAITestBase
 };
 IMPLEMENT_AI_INSTANT_TEST(FWorldConditionTest_EvalComplex, "System.WorldConditions.EvalComplex");
 
-PRAGMA_ENABLE_OPTIMIZATION
-
 #undef LOCTEXT_NAMESPACE
