Fix debug and non editor build issues

diff --git a/Engine/Plugins/Runtime/WorldConditions/Source/WorldConditions/Private/WorldConditionBase.cpp b/Engine/Plugins/Runtime/WorldConditions/Source/WorldConditions/Private/WorldConditionBase.cpp
--- a/Engine/Plugins/Runtime/WorldConditions/Source/WorldConditions/Private/WorldConditionBase.cpp
+++ b/Engine/Plugins/Runtime/WorldConditions/Source/WorldConditions/Private/WorldConditionBase.cpp
@@ -9,10 +9,12 @@ FWorldConditionBase::~FWorldConditionBase()
 	// Empty
 }
 
+#if WITH_EDITOR
 FText FWorldConditionBase::GetDescription() const
 {
 	return StaticStruct()->GetDisplayNameText();
 }
+#endif
 
 bool FWorldConditionBase::Initialize(const UWorldConditionSchema& Schema)
 {
diff --git a/Engine/Plugins/Runtime/WorldConditions/Source/WorldConditions/Private/WorldConditionQuery.cpp b/Engine/Plugins/Runtime/WorldConditions/Source/WorldConditions/Private/WorldConditionQuery.cpp
--- a/Engine/Plugins/Runtime/WorldConditions/Source/WorldConditions/Private/WorldConditionQuery.cpp
+++ b/Engine/Plugins/Runtime/WorldConditions/Source/WorldConditions/Private/WorldConditionQuery.cpp
@@ -286,7 +286,9 @@ bool FWorldConditionQuery::DebugInitialize(const TSubclassOf<UWorldConditionSche
 	}
 
 	QueryDefinition.SchemaClass = InSchemaClass;
+#if WITH_EDITORONLY_DATA
 	QueryDefinition.EditableConditions = InConditions;
+#endif
 	return QueryDefinition.Initialize();
 }
 
diff --git a/Engine/Plugins/Runtime/WorldConditions/Source/WorldConditions/Public/WorldConditionBase.h b/Engine/Plugins/Runtime/WorldConditions/Source/WorldConditions/Public/WorldConditionBase.h
--- a/Engine/Plugins/Runtime/WorldConditions/Source/WorldConditions/Public/WorldConditionBase.h
+++ b/Engine/Plugins/Runtime/WorldConditions/Source/WorldConditions/Public/WorldConditionBase.h
@@ -86,8 +86,10 @@ struct WORLDCONDITIONS_API FWorldConditionBase
 	/** @return The Instance data type of the condition. */
 	virtual const UStruct* GetRuntimeStateType() const { return nullptr; }
 
+#if WITH_EDITOR
 	/** @retrun Description to be shown in the UI. */
 	virtual FText GetDescription() const;
+#endif
 
 	/**
 	 * Initializes the condition to be used with a specific schema.
diff --git a/Engine/Plugins/Runtime/WorldConditions/Source/WorldConditions/Public/WorldConditionContext.h b/Engine/Plugins/Runtime/WorldConditions/Source/WorldConditions/Public/WorldConditionContext.h
--- a/Engine/Plugins/Runtime/WorldConditions/Source/WorldConditions/Public/WorldConditionContext.h
+++ b/Engine/Plugins/Runtime/WorldConditions/Source/WorldConditions/Public/WorldConditionContext.h
@@ -1,4 +1,4 @@
-﻿// Copyright Epic Games, Inc. All Rights Reserved.
+// Copyright Epic Games, Inc. All Rights Reserved.
 
 #pragma once
 
@@ -125,7 +125,7 @@ struct WORLDCONDITIONS_API FWorldConditionContextData
 		const int32 Index = Schema->GetContextDataIndexByName(Name, T::StaticStruct());
 		if (Index != INDEX_NONE)
 		{
-			checkSlow(T::StaticStruct()->IsChildOf(Schema->GetContextDataDescByIndex(Index)->Struct));
+			checkSlow(T::StaticStruct()->IsChildOf(Schema->GetContextDataDescByIndex(Index).Struct));
 			Views[Index] = FWorldConditionDataView(T::StaticStruct(), reinterpret_cast<const uint8*>(Value), Schema->GetContextDataTypeByIndex(Index));
 			return true;
 		}
@@ -140,7 +140,7 @@ struct WORLDCONDITIONS_API FWorldConditionContextData
 		const int32 Index = Schema->GetContextDataIndexByName(Name, T::StaticClass());
 		if (Index != INDEX_NONE)
 		{
-			checkSlow(T::StaticStruct()->IsChildOf(Schema->GetContextDataDescByIndex(Index)->Struct));
+			checkSlow(T::StaticStruct()->IsChildOf(Schema->GetContextDataDescByIndex(Index).Struct));
 			Views[Index] = FWorldConditionDataView(Object, Schema->GetContextDataTypeByIndex(Index));
 			return true;
 		}
