WorldCondition: Fix use of StaticStruct() when StaticClass() should have been used.

diff --git a/Engine/Plugins/Runtime/WorldConditions/Source/WorldConditions/Public/WorldConditionContext.h b/Engine/Plugins/Runtime/WorldConditions/Source/WorldConditions/Public/WorldConditionContext.h
--- a/Engine/Plugins/Runtime/WorldConditions/Source/WorldConditions/Public/WorldConditionContext.h
+++ b/Engine/Plugins/Runtime/WorldConditions/Source/WorldConditions/Public/WorldConditionContext.h
@@ -110,7 +110,7 @@ struct WORLDCONDITIONS_API FWorldConditionContextData
 		check(Schema);
 		if (Ref.IsValid())
 		{
-			checkSlow(T::StaticStruct()->IsChildOf(Schema->GetContextDataDescByRef(Ref)->Struct));
+			checkSlow(T::StaticClass()->IsChildOf(Schema->GetContextDataDescByRef(Ref)->Struct));
 			Views[Ref.GetIndex()] = FWorldConditionDataView(Object, Schema->GetContextDataTypeByRef(Ref));
 			return true;
 		}
@@ -140,7 +140,7 @@ struct WORLDCONDITIONS_API FWorldConditionContextData
 		const int32 Index = Schema->GetContextDataIndexByName(Name, T::StaticClass());
 		if (Index != INDEX_NONE)
 		{
-			checkSlow(T::StaticStruct()->IsChildOf(Schema->GetContextDataDescByIndex(Index).Struct));
+			checkSlow(T::StaticClass()->IsChildOf(Schema->GetContextDataDescByIndex(Index).Struct));
 			Views[Index] = FWorldConditionDataView(Object, Schema->GetContextDataTypeByIndex(Index));
 			return true;
 		}
