Downgrading error to a log when EditableConditions are not available.
Saw issues immediately when loading package ____

diff --git a/Engine/Plugins/Runtime/WorldConditions/Source/WorldConditions/Private/WorldConditionQuery.cpp b/Engine/Plugins/Runtime/WorldConditions/Source/WorldConditions/Private/WorldConditionQuery.cpp
--- a/Engine/Plugins/Runtime/WorldConditions/Source/WorldConditions/Private/WorldConditionQuery.cpp
+++ b/Engine/Plugins/Runtime/WorldConditions/Source/WorldConditions/Private/WorldConditionQuery.cpp
@@ -225,7 +225,7 @@ bool FWorldConditionQueryDefinition::Initialize()
 	}
 	if (ValidConditions.IsEmpty())
 	{
-		UE_LOG(LogWorldCondition, Error, TEXT("World condition requires at least one valid condition."));
+		UE_LOG(LogWorldCondition, Log, TEXT("World condition requires at least one valid condition."));
 		return false;
 	}
 	
