Fixed a bug with the world condition context being cached even when there are conditions that cannot be cached

diff --git a/Engine/Plugins/Runtime/WorldConditions/Source/WorldConditions/Private/WorldConditionContext.cpp b/Engine/Plugins/Runtime/WorldConditions/Source/WorldConditions/Private/WorldConditionContext.cpp
--- a/Engine/Plugins/Runtime/WorldConditions/Source/WorldConditions/Private/WorldConditionContext.cpp
+++ b/Engine/Plugins/Runtime/WorldConditions/Source/WorldConditions/Private/WorldConditionContext.cpp
@@ -50,6 +50,8 @@ bool FWorldConditionContext::IsTrue() const
 	EWorldConditionOperator Operators[UE::WorldCondition::MaxExpressionDepth + 1] = { EWorldConditionOperator::Copy, EWorldConditionOperator::Copy, EWorldConditionOperator::Copy, EWorldConditionOperator::Copy, EWorldConditionOperator::Copy };
 	int32 Depth = 0;
 
+	bool bAllConditionsCanBeCached = true;
+
 	for (int32 Index = 0; Index < QueryState.GetNumConditions(); Index++)
 	{
 		FWorldConditionItem& Item = QueryState.GetItem(Index);
@@ -69,6 +71,8 @@ bool FWorldConditionContext::IsTrue() const
 			{
 				Item.CachedResult = CurrResult;
 			}
+			
+			bAllConditionsCanBeCached &= Condition.bCanCacheResult;
 		}
 
 		Depth++;
@@ -82,9 +86,11 @@ bool FWorldConditionContext::IsTrue() const
 		}
 	}
 
-	QueryState.CachedResult = Results[0]; 
+	const EWorldConditionResult FinalResult = Results[0];
+	
+	QueryState.CachedResult = bAllConditionsCanBeCached ? FinalResult : EWorldConditionResult::Invalid;
 	
-	return QueryState.CachedResult == EWorldConditionResult::IsTrue;
+	return FinalResult == EWorldConditionResult::IsTrue;
 
 }
 
