World Condition: Fix query definition Initialize not  getting called in editor
- fix editor callback to update query definition when edited
- allow to select empty struct for a condition, and display empty
- made the condition display to fit the name (works better with property tables)

diff --git a/Engine/Plugins/Runtime/WorldConditions/Source/WorldConditions/Private/WorldConditionQuery.cpp b/Engine/Plugins/Runtime/WorldConditions/Source/WorldConditions/Private/WorldConditionQuery.cpp
--- a/Engine/Plugins/Runtime/WorldConditions/Source/WorldConditions/Private/WorldConditionQuery.cpp
+++ b/Engine/Plugins/Runtime/WorldConditions/Source/WorldConditions/Private/WorldConditionQuery.cpp
@@ -269,6 +269,20 @@ bool FWorldConditionQueryDefinition::Initialize()
 	return bResult;
 }
 
+void FWorldConditionQueryDefinition::PostSerialize(const FArchive& Ar)
+{
+	if (Ar.IsLoading())
+	{
+#if WITH_EDITOR
+		// Initialize on load in editor.
+		if (Conditions.Num() == 0 && EditableConditions.Num() > 0)
+		{
+			Initialize();
+		}
+#endif
+	}
+}
+
 
 //
 // FWorldConditionQuery
diff --git a/Engine/Plugins/Runtime/WorldConditions/Source/WorldConditions/Public/WorldConditionQuery.h b/Engine/Plugins/Runtime/WorldConditions/Source/WorldConditions/Public/WorldConditionQuery.h
--- a/Engine/Plugins/Runtime/WorldConditions/Source/WorldConditions/Public/WorldConditionQuery.h
+++ b/Engine/Plugins/Runtime/WorldConditions/Source/WorldConditions/Public/WorldConditionQuery.h
@@ -87,6 +87,8 @@ struct WORLDCONDITIONS_API FWorldConditionQueryDefinition
 	/** Initialized the condition from editable data. */
 	bool Initialize();
 
+	void PostSerialize(const FArchive& Ar);
+	
 	/** Conditions of the query, populated by Initialize(). */
 	UPROPERTY()
 	FInstancedStructArray Conditions;
@@ -106,6 +108,15 @@ struct WORLDCONDITIONS_API FWorldConditionQueryDefinition
 	friend struct FWorldConditionContext;
 };
 
+template<>
+struct TStructOpsTypeTraits<FWorldConditionQueryDefinition> : public TStructOpsTypeTraitsBase2<FWorldConditionQueryDefinition>
+{
+	enum
+	{
+		WithPostSerialize = true,
+	};
+};
+
 
 /**
  * Item used to describe the structure of a world condition query for fast traversal of the expression.
diff --git a/Engine/Plugins/Runtime/WorldConditions/Source/WorldConditionsEditor/Private/Customizations/WorldConditionEditableDetails.cpp b/Engine/Plugins/Runtime/WorldConditions/Source/WorldConditionsEditor/Private/Customizations/WorldConditionEditableDetails.cpp
--- a/Engine/Plugins/Runtime/WorldConditions/Source/WorldConditionsEditor/Private/Customizations/WorldConditionEditableDetails.cpp
+++ b/Engine/Plugins/Runtime/WorldConditions/Source/WorldConditionsEditor/Private/Customizations/WorldConditionEditableDetails.cpp
@@ -151,7 +151,7 @@ void FWorldConditionEditableDetails::CustomizeHeader(TSharedRef<IPropertyHandle>
 
 				// Class picker
 				+ SHorizontalBox::Slot()
-				.FillWidth(1.0f)
+				.AutoWidth()
 				.Padding(FMargin(FMargin(4.0f, 0.0f, 0.0f, 0.0f)))
 				.VAlign(VAlign_Center)
 				[
@@ -560,7 +560,7 @@ FText FWorldConditionEditableDetails::GetDisplayValueString() const
 			return Condition->GetDescription();
 		}
 	}
-	return FText();
+	return LOCTEXT("ValueNone", "None");
 }
 
 TSharedRef<SWidget> FWorldConditionEditableDetails::GeneratePicker()
@@ -570,7 +570,7 @@ TSharedRef<SWidget> FWorldConditionEditableDetails::GeneratePicker()
 	StructFilter->Schema = Schema;
 
 	FStructViewerInitializationOptions Options;
-	Options.bShowNoneOption = false;
+	Options.bShowNoneOption = true;
 	Options.StructFilter = StructFilter;
 	Options.NameTypeToDisplay = EStructViewerNameTypeToDisplay::DisplayName;
 	Options.DisplayMode = EStructViewerDisplayMode::ListView;
@@ -606,10 +606,7 @@ void FWorldConditionEditableDetails::OnStructPicked(const UScriptStruct* InStruc
 	{
 		if (FWorldConditionEditable* Condition = static_cast<FWorldConditionEditable*>(Data))
 		{
-			if (InStruct)
-			{
-				Condition->Condition.InitializeAs(InStruct);
-			}
+			Condition->Condition.InitializeAs(InStruct);
 		}
 	}
 
diff --git a/Engine/Plugins/Runtime/WorldConditions/Source/WorldConditionsEditor/Private/Customizations/WorldConditionQueryDefinitionDetails.cpp b/Engine/Plugins/Runtime/WorldConditions/Source/WorldConditionsEditor/Private/Customizations/WorldConditionQueryDefinitionDetails.cpp
--- a/Engine/Plugins/Runtime/WorldConditions/Source/WorldConditionsEditor/Private/Customizations/WorldConditionQueryDefinitionDetails.cpp
+++ b/Engine/Plugins/Runtime/WorldConditions/Source/WorldConditionsEditor/Private/Customizations/WorldConditionQueryDefinitionDetails.cpp
@@ -26,19 +26,8 @@ void FWorldConditionQueryDefinitionDetails::CustomizeHeader(TSharedRef<IProperty
 	EditableConditionsProperty = StructProperty->GetChildHandle(TEXT("EditableConditions"));
 
 	// Keep the definition up to date as it's being edited.
-	StructProperty->SetOnPropertyValueChanged(FSimpleDelegate::CreateLambda(
-	[StructProperty = StructProperty]()
-	{
-		TArray<void*> RawNodeData;
-		StructProperty->AccessRawData(RawNodeData);
-		for (void* Data : RawNodeData)
-		{
-			if (FWorldConditionQueryDefinition* QueryDefinition = static_cast<FWorldConditionQueryDefinition*>(Data))
-			{
-				QueryDefinition->Initialize();
-			}
-		}
-	}));
+	StructProperty->SetOnPropertyValueChanged(FSimpleDelegate::CreateSP(this, &FWorldConditionQueryDefinitionDetails::InitializeDefinition));
+	StructProperty->SetOnChildPropertyValueChanged(FSimpleDelegate::CreateSP(this, &FWorldConditionQueryDefinitionDetails::InitializeDefinition));
 	
 	HeaderRow
 		.NameContent()
@@ -67,4 +56,19 @@ void FWorldConditionQueryDefinitionDetails::CustomizeChildren(TSharedRef<IProper
 	}
 }
 
+void FWorldConditionQueryDefinitionDetails::InitializeDefinition() const
+{
+	check(StructProperty);
+	TArray<void*> RawNodeData;
+	StructProperty->AccessRawData(RawNodeData);
+	for (void* Data : RawNodeData)
+	{
+		if (FWorldConditionQueryDefinition* QueryDefinition = static_cast<FWorldConditionQueryDefinition*>(Data))
+		{
+			QueryDefinition->Initialize();
+		}
+	}
+}
+
+
 #undef LOCTEXT_NAMESPACE
diff --git a/Engine/Plugins/Runtime/WorldConditions/Source/WorldConditionsEditor/Private/Customizations/WorldConditionQueryDefinitionDetails.h b/Engine/Plugins/Runtime/WorldConditions/Source/WorldConditionsEditor/Private/Customizations/WorldConditionQueryDefinitionDetails.h
--- a/Engine/Plugins/Runtime/WorldConditions/Source/WorldConditionsEditor/Private/Customizations/WorldConditionQueryDefinitionDetails.h
+++ b/Engine/Plugins/Runtime/WorldConditions/Source/WorldConditionsEditor/Private/Customizations/WorldConditionQueryDefinitionDetails.h
@@ -25,6 +25,8 @@ private:
 	virtual void CustomizeHeader(TSharedRef<IPropertyHandle> StructPropertyHandle, FDetailWidgetRow& HeaderRow, IPropertyTypeCustomizationUtils& StructCustomizationUtils) override;
 	virtual void CustomizeChildren(TSharedRef<IPropertyHandle> StructPropertyHandle, IDetailChildrenBuilder& StructBuilder, IPropertyTypeCustomizationUtils& StructCustomizationUtils) override;
 
+	void InitializeDefinition() const;
+	
 	IPropertyUtilities* PropUtils = nullptr;
 	TSharedPtr<IPropertyHandle> StructProperty;
 	TSharedPtr<IPropertyHandle> EditableConditionsProperty;
