World Conditions: Remove SetSchema() pointer version

diff --git a/Engine/Plugins/Runtime/WorldConditions/Source/WorldConditions/Public/WorldConditionContext.h b/Engine/Plugins/Runtime/WorldConditions/Source/WorldConditions/Public/WorldConditionContext.h
--- a/Engine/Plugins/Runtime/WorldConditions/Source/WorldConditions/Public/WorldConditionContext.h
+++ b/Engine/Plugins/Runtime/WorldConditions/Source/WorldConditions/Public/WorldConditionContext.h
@@ -82,13 +82,6 @@ struct WORLDCONDITIONS_API FWorldConditionContextData
 		}
 	}
 
-	/** @todo: this is temporary change to allow non-engine code to migrate to ref API. */
-	void SetSchema(const UWorldConditionSchema* InSchema)
-	{
-		check(InSchema);
-		SetSchema(*InSchema);
-	}
-
 	/** Sets context data Struct at location specified by Ref. */
 	template <typename T>
 	typename TEnableIf<!TIsDerivedFrom<T, UObject>::IsDerived, bool>::Type SetContextData(const FWorldConditionContextDataRef& Ref, const T* Value)
