World Condition: Added logic and UI to invert a condition

diff --git a/Engine/Plugins/Runtime/WorldConditions/Source/WorldConditions/Private/WorldConditionContext.cpp b/Engine/Plugins/Runtime/WorldConditions/Source/WorldConditions/Private/WorldConditionContext.cpp
--- a/Engine/Plugins/Runtime/WorldConditions/Source/WorldConditions/Private/WorldConditionContext.cpp
+++ b/Engine/Plugins/Runtime/WorldConditions/Source/WorldConditions/Private/WorldConditionContext.cpp
@@ -67,6 +67,8 @@ bool FWorldConditionContext::IsTrue() const
 			check(QueryDefinition.Conditions.Num() == QueryState.GetNumConditions());
 			const FWorldConditionBase& Condition = QueryDefinition.Conditions[Index].Get<FWorldConditionBase>();
 			CurrResult = Condition.IsTrue(*this);
+			CurrResult = UE::WorldCondition::Invert(CurrResult, Condition.bInvert); 
+			
 			if (Condition.bCanCacheResult)
 			{
 				Item.CachedResult = CurrResult;
diff --git a/Engine/Plugins/Runtime/WorldConditions/Source/WorldConditions/Private/WorldConditionQuery.cpp b/Engine/Plugins/Runtime/WorldConditions/Source/WorldConditions/Private/WorldConditionQuery.cpp
--- a/Engine/Plugins/Runtime/WorldConditions/Source/WorldConditions/Private/WorldConditionQuery.cpp
+++ b/Engine/Plugins/Runtime/WorldConditions/Source/WorldConditions/Private/WorldConditionQuery.cpp
@@ -221,6 +221,7 @@ bool FWorldConditionQueryDefinition::Initialize()
 				// Store expression depth temporarily into NextExpressionDepth, it will be update below.
 				Condition.NextExpressionDepth = EditableCondition.ExpressionDepth;
 				Condition.Operator = EditableCondition.Operator;
+				Condition.bInvert = EditableCondition.bInvert;
 
 				ValidConditions.Add(EditableCondition.Condition);
 			}
diff --git a/Engine/Plugins/Runtime/WorldConditions/Source/WorldConditions/Public/WorldConditionBase.h b/Engine/Plugins/Runtime/WorldConditions/Source/WorldConditions/Public/WorldConditionBase.h
--- a/Engine/Plugins/Runtime/WorldConditions/Source/WorldConditions/Public/WorldConditionBase.h
+++ b/Engine/Plugins/Runtime/WorldConditions/Source/WorldConditions/Public/WorldConditionBase.h
@@ -79,6 +79,7 @@ struct WORLDCONDITIONS_API FWorldConditionBase
 	FWorldConditionBase()
 		: bIsStateObject(false)
 		, bCanCacheResult(true)
+		, bInvert(false)
 	{
 	}
 	virtual ~FWorldConditionBase();
@@ -140,6 +141,10 @@ protected:
 	/** Set to true if the result of the IsTrue() can be cached. */
 	uint8 bCanCacheResult : 1;
 
+	/** Controls whether the value of the expressions as calculated by IsTrue should be inverted. The inversion is handled by the expression evaluator. */
+	UPROPERTY()
+	uint8 bInvert : 1;
+
 	/** Operator describing how the results of the condition is combined with other conditions. Not used directly, but to set up condition item in query state. */
 	UPROPERTY()
 	EWorldConditionOperator Operator = EWorldConditionOperator::And;
diff --git a/Engine/Plugins/Runtime/WorldConditions/Source/WorldConditions/Public/WorldConditionQuery.h b/Engine/Plugins/Runtime/WorldConditions/Source/WorldConditions/Public/WorldConditionQuery.h
--- a/Engine/Plugins/Runtime/WorldConditions/Source/WorldConditions/Public/WorldConditionQuery.h
+++ b/Engine/Plugins/Runtime/WorldConditions/Source/WorldConditions/Public/WorldConditionQuery.h
@@ -49,7 +49,15 @@ struct WORLDCONDITIONS_API FWorldConditionEditable
 		, Condition(InCondition)
 	{
 	}
-	
+
+	FWorldConditionEditable(const uint8 InExpressionDepth, const EWorldConditionOperator InOperator, const bool bInInvert, const FConstStructView InCondition)
+		: ExpressionDepth(InExpressionDepth)
+		, Operator(InOperator)
+		, bInvert(bInInvert)
+		, Condition(InCondition)
+		{
+		}
+
 	void Reset()
 	{
 		Operator = EWorldConditionOperator::And;
@@ -64,7 +72,11 @@ struct WORLDCONDITIONS_API FWorldConditionEditable
 	/** Operator describing how the results of the condition is combined with other conditions. */
 	UPROPERTY(EditAnywhere, Category="Default")
 	EWorldConditionOperator Operator = EWorldConditionOperator::And;
-	
+
+	/** Controls whether the value of the expressions as calculated by IsTrue should be inverted. */
+	UPROPERTY(EditAnywhere, Category="Default")
+	bool bInvert = false;
+
 	/** Instance of a world condition. */
 	UPROPERTY(EditAnywhere, Category="Default")
 	FInstancedStruct Condition;
diff --git a/Engine/Plugins/Runtime/WorldConditions/Source/WorldConditions/Public/WorldConditionTypes.h b/Engine/Plugins/Runtime/WorldConditions/Source/WorldConditions/Public/WorldConditionTypes.h
--- a/Engine/Plugins/Runtime/WorldConditions/Source/WorldConditions/Public/WorldConditionTypes.h
+++ b/Engine/Plugins/Runtime/WorldConditions/Source/WorldConditions/Public/WorldConditionTypes.h
@@ -47,6 +47,17 @@ namespace UE::WorldCondition
 		const uint8 Result = (Operator == EWorldConditionOperator::And) ? ResultAnd : ResultOr;
 		return (Operator == EWorldConditionOperator::Copy) ? InRight : (EWorldConditionResult)Result;
 	}
+
+	/** Inverts the result if bInvert is true. Invalid is kept intact. */
+	inline EWorldConditionResult Invert(const EWorldConditionResult InResult, const bool bInvert)
+	{
+		if (InResult == EWorldConditionResult::Invalid)
+		{
+			return EWorldConditionResult::Invalid;
+		}
+		return ((InResult == EWorldConditionResult::IsTrue) ^ bInvert) ? EWorldConditionResult::IsTrue : EWorldConditionResult::IsFalse;
+	}
+
 } // UE::WorldCondition
 
 /**
diff --git a/Engine/Plugins/Runtime/WorldConditions/Source/WorldConditionsEditor/Private/Customizations/WorldConditionEditableDetails.cpp b/Engine/Plugins/Runtime/WorldConditions/Source/WorldConditionsEditor/Private/Customizations/WorldConditionEditableDetails.cpp
--- a/Engine/Plugins/Runtime/WorldConditions/Source/WorldConditionsEditor/Private/Customizations/WorldConditionEditableDetails.cpp
+++ b/Engine/Plugins/Runtime/WorldConditions/Source/WorldConditionsEditor/Private/Customizations/WorldConditionEditableDetails.cpp
@@ -19,6 +19,7 @@
 #include "WorldConditionSchema.h"
 #include "Editor.h"
 #include "Modules/ModuleManager.h"
+#include "Widgets/Input/SCheckBox.h"
 
 #define LOCTEXT_NAMESPACE "WorldCondition"
 
@@ -70,10 +71,12 @@ void FWorldConditionEditableDetails::CustomizeHeader(TSharedRef<IPropertyHandle>
 	ConditionProperty = StructProperty->GetChildHandle(TEXT("Condition"));
 	ExpressionDepthProperty = StructProperty->GetChildHandle(TEXT("ExpressionDepth"));
 	OperatorProperty = StructProperty->GetChildHandle(TEXT("Operator"));
+	InvertProperty = StructProperty->GetChildHandle(TEXT("bInvert"));
 
 	check(ConditionProperty.IsValid());
 	check(ExpressionDepthProperty.IsValid());
 	check(OperatorProperty.IsValid());
+	check(InvertProperty.IsValid());
 	
 	CacheSchema();
 
@@ -149,6 +152,22 @@ void FWorldConditionEditableDetails::CustomizeHeader(TSharedRef<IPropertyHandle>
 					.Text(this, &FWorldConditionEditableDetails::GetOpenParens)
 				]
 
+				// Not
+				+ SHorizontalBox::Slot()
+				.AutoWidth()
+				.VAlign(VAlign_Center)
+				[
+					SNew(SBox)
+					.Padding(FMargin(8, 0, 4, 0))
+					.VAlign(VAlign_Center)
+					.Visibility(this, &FWorldConditionEditableDetails::GetInvertVisibility)
+					[
+						SNew(STextBlock)
+						.TextStyle(FWorldConditionEditorStyle::Get(), "Condition.Operator")
+						.Text(LOCTEXT("NotOperator","NOT"))
+					]
+				]
+				
 				// Class picker
 				+ SHorizontalBox::Slot()
 				.AutoWidth()
@@ -188,11 +207,47 @@ void FWorldConditionEditableDetails::CustomizeHeader(TSharedRef<IPropertyHandle>
 void FWorldConditionEditableDetails::CustomizeChildren(TSharedRef<IPropertyHandle> StructPropertyHandle, IDetailChildrenBuilder& StructBuilder, IPropertyTypeCustomizationUtils& StructCustomizationUtils)
 {
 	check(ConditionProperty);
+	check(InvertProperty);
+
+	StructBuilder.AddProperty(InvertProperty.ToSharedRef());
 	
 	TSharedRef<FInstancedStructDataDetails> NodeDetails = MakeShareable(new FInstancedStructDataDetails(ConditionProperty));
 	StructBuilder.AddCustomBuilder(NodeDetails);
 }
 
+void FWorldConditionEditableDetails::CacheSchema()
+{
+	// Find schema from outer FWorldConditionQueryDefinition.
+	Schema = nullptr;
+	
+	TSharedPtr<IPropertyHandle> CurrentProperty = StructProperty;
+	while (CurrentProperty.IsValid() && !Schema)
+	{
+		const FProperty* Property = CurrentProperty->GetProperty();
+		if (const FStructProperty* CurrentStructProperty = CastField<FStructProperty>(Property))
+		{
+			if (CurrentStructProperty->Struct == TBaseStructure<FWorldConditionQueryDefinition>::Get())
+			{
+				// Get schema from definition
+				TArray<void*> RawNodeData;
+				CurrentProperty->AccessRawData(RawNodeData);
+				for (void* Data : RawNodeData)
+				{
+					if (const FWorldConditionQueryDefinition* QueryDefinition = static_cast<FWorldConditionQueryDefinition*>(Data))
+					{
+						Schema = QueryDefinition->SchemaClass.GetDefaultObject();
+						if (Schema)
+						{
+							break;
+						}
+					}
+				}
+			}
+		}
+		CurrentProperty = CurrentProperty->GetParentHandle(); 
+	}
+}
+
 bool FWorldConditionEditableDetails::ShouldResetToDefault(TSharedPtr<IPropertyHandle> PropertyHandle) const
 {
 	check(StructProperty);
@@ -247,37 +302,14 @@ void FWorldConditionEditableDetails::ResetToDefault(TSharedPtr<IPropertyHandle>
 	}
 }
 
-void FWorldConditionEditableDetails::CacheSchema()
+EVisibility FWorldConditionEditableDetails::GetInvertVisibility() const
 {
-	// Find schema from outer FWorldConditionQueryDefinition.
-	Schema = nullptr;
+	check(InvertProperty);
 	
-	TSharedPtr<IPropertyHandle> CurrentProperty = StructProperty;
-	while (CurrentProperty.IsValid() && !Schema)
-	{
-		const FProperty* Property = CurrentProperty->GetProperty();
-		if (const FStructProperty* CurrentStructProperty = CastField<FStructProperty>(Property))
-		{
-			if (CurrentStructProperty->Struct == TBaseStructure<FWorldConditionQueryDefinition>::Get())
-			{
-				// Get schema from definition
-				TArray<void*> RawNodeData;
-				CurrentProperty->AccessRawData(RawNodeData);
-				for (void* Data : RawNodeData)
-				{
-					if (const FWorldConditionQueryDefinition* QueryDefinition = static_cast<FWorldConditionQueryDefinition*>(Data))
-					{
-						Schema = QueryDefinition->SchemaClass.GetDefaultObject();
-						if (Schema)
-						{
-							break;
-						}
-					}
-				}
-			}
-		}
-		CurrentProperty = CurrentProperty->GetParentHandle(); 
-	}
+	bool bInvert = false;
+	InvertProperty->GetValue(bInvert);
+	
+	return bInvert ? EVisibility::Visible : EVisibility::Collapsed;
 }
 
 FOptionalSize FWorldConditionEditableDetails::GetIndentSize() const
diff --git a/Engine/Plugins/Runtime/WorldConditions/Source/WorldConditionsEditor/Private/Customizations/WorldConditionEditableDetails.h b/Engine/Plugins/Runtime/WorldConditions/Source/WorldConditionsEditor/Private/Customizations/WorldConditionEditableDetails.h
--- a/Engine/Plugins/Runtime/WorldConditions/Source/WorldConditionsEditor/Private/Customizations/WorldConditionEditableDetails.h
+++ b/Engine/Plugins/Runtime/WorldConditions/Source/WorldConditionsEditor/Private/Customizations/WorldConditionEditableDetails.h
@@ -32,7 +32,9 @@ private:
 
 	bool ShouldResetToDefault(TSharedPtr<IPropertyHandle> PropertyHandle) const;
 	void ResetToDefault(TSharedPtr<IPropertyHandle> PropertyHandle);
-	
+
+	EVisibility GetInvertVisibility() const;
+
 	FOptionalSize GetIndentSize() const;
 	TSharedRef<SWidget> OnGetIndentContent() const;
 
@@ -69,5 +71,6 @@ private:
 	TSharedPtr<IPropertyHandle> ConditionProperty;
 	TSharedPtr<IPropertyHandle> ExpressionDepthProperty;
 	TSharedPtr<IPropertyHandle> OperatorProperty;
+	TSharedPtr<IPropertyHandle> InvertProperty;
 	IPropertyUtilities* PropUtils = nullptr;
 };
diff --git a/Engine/Plugins/Runtime/WorldConditions/Source/WorldConditionsTestSuite/Private/WorldConditionsTest.cpp b/Engine/Plugins/Runtime/WorldConditions/Source/WorldConditionsTestSuite/Private/WorldConditionsTest.cpp
--- a/Engine/Plugins/Runtime/WorldConditions/Source/WorldConditionsTestSuite/Private/WorldConditionsTest.cpp
+++ b/Engine/Plugins/Runtime/WorldConditions/Source/WorldConditionsTestSuite/Private/WorldConditionsTest.cpp
@@ -10,6 +10,8 @@
 
 #define LOCTEXT_NAMESPACE "AITestSuite_WorldConditionsTest"
 
+UE_DISABLE_OPTIMIZATION_SHIP
+
 struct FWorldConditionTest_Init : FAITestBase
 {
 	virtual bool InstantTest() override
@@ -71,6 +73,41 @@ struct FWorldConditionTest_Eval : FAITestBase
 };
 IMPLEMENT_AI_INSTANT_TEST(FWorldConditionTest_Eval, "System.WorldConditions.Eval");
 
+struct FWorldConditionTest_EvalInvert : FAITestBase
+{
+	virtual bool InstantTest() override
+	{
+		FWorldConditionQuery Query;
+		const bool bInitialized = Query.DebugInitialize(UWorldConditionTestSchema::StaticClass(),
+			{
+				FWorldConditionEditable(0, EWorldConditionOperator::Copy, FConstStructView::Make(FWorldConditionTest(1))),
+				FWorldConditionEditable(0, EWorldConditionOperator::And, /*bInveret*/true, FConstStructView::Make(FWorldConditionTest(0)))
+			});
+
+		AITEST_TRUE("Query should get initialized", bInitialized);
+		
+		const FWorldConditionTestData TestData(1);
+
+		const UWorldConditionTestSchema* TestSchema = GetDefault<UWorldConditionTestSchema>();
+		FWorldConditionContextData ContextData(*TestSchema);
+		const bool bContextDataSet = ContextData.SetContextData(TestSchema->GetValueRef(), &TestData);
+		AITEST_TRUE("Query context data should be set", bContextDataSet);
+		
+		const bool bActivated = Query.Activate(GetWorld(), ContextData);
+		AITEST_TRUE("Query should get activated", bActivated);
+		AITEST_TRUE("Query should be active", Query.IsActive());
+		
+		const bool bResult = Query.IsTrue(ContextData);
+		AITEST_TRUE("Query result should be true", bResult);
+
+		Query.Deactivate(ContextData);
+		AITEST_FALSE("Query should not be active", Query.IsActive());
+
+		return true;
+	}
+};
+IMPLEMENT_AI_INSTANT_TEST(FWorldConditionTest_EvalInvert, "System.WorldConditions.EvalInvert");
+
 struct FWorldConditionTest_CachedEval : FAITestBase
 {
 	virtual bool InstantTest() override
@@ -232,4 +269,6 @@ struct FWorldConditionTest_FailingActivate : FAITestBase
 };
 IMPLEMENT_AI_INSTANT_TEST(FWorldConditionTest_FailingActivate, "System.WorldConditions.FailingActivate");
 
+UE_ENABLE_OPTIMIZATION_SHIP
+
 #undef LOCTEXT_NAMESPACE
