WorldConditions: Remove code used during API migration.

diff --git a/Engine/Plugins/Runtime/WorldConditions/Source/WorldConditions/Public/WorldConditionContext.h b/Engine/Plugins/Runtime/WorldConditions/Source/WorldConditions/Public/WorldConditionContext.h
--- a/Engine/Plugins/Runtime/WorldConditions/Source/WorldConditions/Public/WorldConditionContext.h
+++ b/Engine/Plugins/Runtime/WorldConditions/Source/WorldConditions/Public/WorldConditionContext.h
@@ -155,16 +155,6 @@ protected:
  */
 struct WORLDCONDITIONS_API FWorldConditionContext
 {
-	// @todo Remove, dummy method to help to migrate to the new API, will be removed immeidately in follow up change.
-	explicit FWorldConditionContext(const UObject& Owner, const FWorldConditionQueryDefinition& Defintion, FWorldConditionQueryState& InQueryState, const FWorldConditionContextData& InContextData)
-		: Owner(*InQueryState.Owner)
-		, QueryState(InQueryState)
-		, ContextData(InContextData)
-	{
-		check(InQueryState.Owner);
-		World = Owner.GetWorld();
-	}
-	
 	explicit FWorldConditionContext(FWorldConditionQueryState& InQueryState, const FWorldConditionContextData& InContextData)
 		: Owner(*InQueryState.Owner)
 		, QueryState(InQueryState)
diff --git a/Engine/Plugins/Runtime/WorldConditions/Source/WorldConditions/Public/WorldConditionQuery.h b/Engine/Plugins/Runtime/WorldConditions/Source/WorldConditions/Public/WorldConditionQuery.h
--- a/Engine/Plugins/Runtime/WorldConditions/Source/WorldConditions/Public/WorldConditionQuery.h
+++ b/Engine/Plugins/Runtime/WorldConditions/Source/WorldConditions/Public/WorldConditionQuery.h
@@ -323,13 +323,6 @@ struct WORLDCONDITIONS_API FWorldConditionQueryState
 	 * @return Invalidation handle.	
 	 */
 	FWorldConditionResultInvalidationHandle GetInvalidationHandle(const FWorldConditionBase& Condition) const;
-
-	// @todo Remove, dummy method to help to migrate to the new API, will be removed immeidately in follow up change.
-	void AddReferencedObjects(const FWorldConditionQueryDefinition& Definition, FReferenceCollector& Collector) const
-	{
-		FWorldConditionQueryState* NonConstThis = const_cast<FWorldConditionQueryState*>(this);
-		NonConstThis->AddStructReferencedObjects(Collector);
-	}
 	
 protected:
 
