World Conditions: Use IsChildOf() in GetContextDataIndexByName

diff --git a/Engine/Plugins/Runtime/WorldConditions/Source/WorldConditions/Private/WorldConditionSchema.cpp b/Engine/Plugins/Runtime/WorldConditions/Source/WorldConditions/Private/WorldConditionSchema.cpp
--- a/Engine/Plugins/Runtime/WorldConditions/Source/WorldConditions/Private/WorldConditionSchema.cpp
+++ b/Engine/Plugins/Runtime/WorldConditions/Source/WorldConditions/Private/WorldConditionSchema.cpp
@@ -34,7 +34,7 @@ int32 UWorldConditionSchema::GetContextDataIndexByName(const FName DataName, con
 {
 	return ContextDataDescs.IndexOfByPredicate([DataName, Struct](const FWorldConditionContextDataDesc& Desc)
 	{
-		return Desc.Name == DataName && Desc.Struct == Struct;
+		return Desc.Name == DataName && Desc.Struct->IsChildOf(Struct);
 	});
 }
 
@@ -42,7 +42,7 @@ const FWorldConditionContextDataDesc* UWorldConditionSchema::GetContextDataDescB
 {
 	return ContextDataDescs.FindByPredicate([DataName, Struct](const FWorldConditionContextDataDesc& Desc)
 	{
-		return Desc.Name == DataName && Desc.Struct == Struct;
+		return Desc.Name == DataName && Desc.Struct->IsChildOf(Struct);
 	});
 }
 
