World Condition: Use package to find outer object for data tables.

diff --git a/Engine/Plugins/Runtime/WorldConditions/Source/WorldConditionsEditor/Private/Customizations/WorldConditionQueryDefinitionDetails.cpp b/Engine/Plugins/Runtime/WorldConditions/Source/WorldConditionsEditor/Private/Customizations/WorldConditionQueryDefinitionDetails.cpp
--- a/Engine/Plugins/Runtime/WorldConditions/Source/WorldConditionsEditor/Private/Customizations/WorldConditionQueryDefinitionDetails.cpp
+++ b/Engine/Plugins/Runtime/WorldConditions/Source/WorldConditionsEditor/Private/Customizations/WorldConditionQueryDefinitionDetails.cpp
@@ -60,18 +60,39 @@ void FWorldConditionQueryDefinitionDetails::InitializeDefinition() const
 {
 	check(StructProperty);
 
-	TArray<void*> RawNodeData;
+	TArray<void*> RawData;
 	TArray<UObject*> OuterObjects;
 		
-	StructProperty->AccessRawData(RawNodeData);
+	StructProperty->AccessRawData(RawData);
 	StructProperty->GetOuterObjects(OuterObjects);
-	check(RawNodeData.Num() == OuterObjects.Num());
-		
-	for (int32 Index = 0; Index < RawNodeData.Num(); Index++)
+
+	// If no outer objects, try to find an outer asset (can happen e.g. in data table).
+	if (OuterObjects.Num() == 0)
+	{
+		TArray<UPackage*> OuterPackages;
+		StructProperty->GetOuterPackages(OuterPackages);
+		for (UPackage* Package : OuterPackages)
+		{
+			if (UObject* Asset = Package->FindAssetInPackage())
+			{
+				OuterObjects.Add(Asset);
+			}
+		}
+		// Number of objects from the packages should match number of data.
+		if (OuterObjects.Num() != RawData.Num())
+		{
+			ensureMsgf(false, TEXT("Cannot edit World Condition Definition in this context, cannot find a valid outer object."));
+			return;
+		}
+	}
+
+	check(RawData.Num() == OuterObjects.Num());
+	
+	for (int32 Index = 0; Index < RawData.Num(); Index++)
 	{
 		if (UObject* Outer = OuterObjects[Index])
 		{
-			if (FWorldConditionQueryDefinition* QueryDefinition = static_cast<FWorldConditionQueryDefinition*>(RawNodeData[Index]))
+			if (FWorldConditionQueryDefinition* QueryDefinition = static_cast<FWorldConditionQueryDefinition*>(RawData[Index]))
 			{
 				QueryDefinition->Initialize(*Outer);
 			}
