World Conditions: Call condition Initialize() always on postload.

diff --git a/Engine/Plugins/Runtime/WorldConditions/Source/WorldConditions/Private/WorldConditionQuery.cpp b/Engine/Plugins/Runtime/WorldConditions/Source/WorldConditions/Private/WorldConditionQuery.cpp
--- a/Engine/Plugins/Runtime/WorldConditions/Source/WorldConditions/Private/WorldConditionQuery.cpp
+++ b/Engine/Plugins/Runtime/WorldConditions/Source/WorldConditions/Private/WorldConditionQuery.cpp
@@ -251,6 +251,40 @@ FWorldConditionResultInvalidationHandle FWorldConditionQueryState::GetInvalidati
 	return FWorldConditionResultInvalidationHandle(CachedResult, Item);
 }
 
+
+//
+// UWorldConditionQuerySharedDefinition
+//
+
+void UWorldConditionQuerySharedDefinition::PostLoad()
+{
+	Super::PostLoad();
+	
+	// Initialize conditions. This is done always on load so that any changes to the schema take affect.
+	if (SchemaClass)
+	{
+		const UWorldConditionSchema* Schema = SchemaClass.GetDefaultObject();
+		bool bResult = true;
+
+		for (int32 Index = 0; Index < Conditions.Num(); Index++)
+		{
+			FWorldConditionBase& Condition = Conditions[Index].GetMutable<FWorldConditionBase>();
+			bResult &= Condition.Initialize(*Schema);
+		}
+
+		if (!bResult)
+		{
+			UE_LOG(LogWorldCondition, Error, TEXT("World Condition: Failed to initialize query %s for %s."),
+				*GetNameSafe(this), *GetFullNameSafe(GetOuter()));
+		}
+	}
+	else
+	{
+		UE_LOG(LogWorldCondition, Error, TEXT("World Condition: shared definition %s for %s has empty schema, and %d conditions."),
+			*GetNameSafe(this), *GetFullNameSafe(GetOuter()), Conditions.Num());
+	}
+}
+
 //
 // FWorldConditionQueryDefinition
 //
diff --git a/Engine/Plugins/Runtime/WorldConditions/Source/WorldConditions/Public/WorldConditionQuery.h b/Engine/Plugins/Runtime/WorldConditions/Source/WorldConditions/Public/WorldConditionQuery.h
--- a/Engine/Plugins/Runtime/WorldConditions/Source/WorldConditions/Public/WorldConditionQuery.h
+++ b/Engine/Plugins/Runtime/WorldConditions/Source/WorldConditions/Public/WorldConditionQuery.h
@@ -94,6 +94,8 @@ class WORLDCONDITIONS_API UWorldConditionQuerySharedDefinition : public UObject
 public:
 	GENERATED_BODY()
 
+	virtual void PostLoad() override;
+	
 	/** All the conditions of the world conditions. */
 	UPROPERTY()
 	FInstancedStructContainer Conditions;
