World Conditions: Fix mutable uses of instanced structs and struct views.

diff --git a/Engine/Plugins/Runtime/WorldConditions/Source/WorldConditions/Private/WorldConditionQuery.cpp b/Engine/Plugins/Runtime/WorldConditions/Source/WorldConditions/Private/WorldConditionQuery.cpp
--- a/Engine/Plugins/Runtime/WorldConditions/Source/WorldConditions/Private/WorldConditionQuery.cpp
+++ b/Engine/Plugins/Runtime/WorldConditions/Source/WorldConditions/Private/WorldConditionQuery.cpp
@@ -312,7 +312,7 @@ bool FWorldConditionQueryDefinition::Initialize(UObject& Outer)
 	// Append only valid condition.
 	TArray<FConstStructView> ValidConditions;
 	ValidConditions.Reserve(EditableConditions.Num());
-	for (const FWorldConditionEditable& EditableCondition : EditableConditions)
+	for (FWorldConditionEditable& EditableCondition : EditableConditions)
 	{
 		if (EditableCondition.Condition.IsValid())
 		{
