World Conditions: Add setter for schema in the definition.

diff --git a/Engine/Plugins/Runtime/WorldConditions/Source/WorldConditions/Private/WorldConditionQuery.cpp b/Engine/Plugins/Runtime/WorldConditions/Source/WorldConditions/Private/WorldConditionQuery.cpp
--- a/Engine/Plugins/Runtime/WorldConditions/Source/WorldConditions/Private/WorldConditionQuery.cpp
+++ b/Engine/Plugins/Runtime/WorldConditions/Source/WorldConditions/Private/WorldConditionQuery.cpp
@@ -288,6 +288,11 @@ void UWorldConditionQuerySharedDefinition::PostLoad()
 // FWorldConditionQueryDefinition
 //
 
+void FWorldConditionQueryDefinition::SetSchemaClass(const TSubclassOf<UWorldConditionSchema> InSchema)
+{
+	SchemaClass = InSchema;
+}
+
 bool FWorldConditionQueryDefinition::IsValid() const
 {
 	return SharedDefinition != nullptr;
diff --git a/Engine/Plugins/Runtime/WorldConditions/Source/WorldConditions/Public/WorldConditionQuery.h b/Engine/Plugins/Runtime/WorldConditions/Source/WorldConditions/Public/WorldConditionQuery.h
--- a/Engine/Plugins/Runtime/WorldConditions/Source/WorldConditions/Public/WorldConditionQuery.h
+++ b/Engine/Plugins/Runtime/WorldConditions/Source/WorldConditions/Public/WorldConditionQuery.h
@@ -114,6 +114,9 @@ struct WORLDCONDITIONS_API FWorldConditionQueryDefinition
 {
 	GENERATED_BODY()
 
+	/** Sets the schema class that is used to create the conditions. */
+	void SetSchemaClass(const TSubclassOf<UWorldConditionSchema> InSchema);
+
 	/** @return true of the definition has conditions and has been initialized. */
 	bool IsValid() const;
 
