[WorldConditionQuery] limit call to SerializeTaggedProperties only for archive configurations supported by UStruct::SerializeVersionedTaggedProperties.

diff --git a/Engine/Plugins/Runtime/WorldConditions/Source/WorldConditions/Private/WorldConditionQuery.cpp b/Engine/Plugins/Runtime/WorldConditions/Source/WorldConditions/Private/WorldConditionQuery.cpp
--- a/Engine/Plugins/Runtime/WorldConditions/Source/WorldConditions/Private/WorldConditionQuery.cpp
+++ b/Engine/Plugins/Runtime/WorldConditions/Source/WorldConditions/Private/WorldConditionQuery.cpp
@@ -479,8 +479,12 @@ bool FWorldConditionQueryDefinition::Serialize(FArchive& Ar)
 {
 	Ar.UsingCustomVersion(FWorldConditionCustomVersion::GUID);
 
-	// Use default serialization for most properties.
-	StaticStruct()->SerializeTaggedProperties(Ar, (uint8*)this, StaticStruct(), nullptr);
+	// Use default serialization for most properties for
+	// archive configurations supported by UStruct::SerializeVersionedTaggedProperties 
+	if (Ar.IsLoading() || Ar.IsSaving() || Ar.IsCountingMemory() || Ar.IsObjectReferenceCollector())
+	{
+		StaticStruct()->SerializeTaggedProperties(Ar, (uint8*)this, StaticStruct(), nullptr);
+	}
 
 	if (Ar.CustomVer(FWorldConditionCustomVersion::GUID) >= FWorldConditionCustomVersion::StructSharedDefinition)
 	{
@@ -500,7 +504,8 @@ bool FWorldConditionQueryDefinition::Serialize(FArchive& Ar)
 			}
 		}
 		
-		if (SharedDefinition.IsValid())
+		if (SharedDefinition.IsValid()
+			&& (Ar.IsLoading() || Ar.IsSaving() || Ar.IsCountingMemory() || Ar.IsObjectReferenceCollector()))
 		{
 			UScriptStruct* Struct = TBaseStructure<FWorldConditionQuerySharedDefinition>::Get();
 			Struct->SerializeTaggedProperties(Ar, (uint8*)SharedDefinition.Get(), Struct, nullptr);
