World Conditions: Fix linking of FWorldConditionQuerySharedDefinition on load

diff --git a/Engine/Plugins/Runtime/WorldConditions/Source/WorldConditions/Private/WorldConditionQuery.cpp b/Engine/Plugins/Runtime/WorldConditions/Source/WorldConditions/Private/WorldConditionQuery.cpp
--- a/Engine/Plugins/Runtime/WorldConditions/Source/WorldConditions/Private/WorldConditionQuery.cpp
+++ b/Engine/Plugins/Runtime/WorldConditions/Source/WorldConditions/Private/WorldConditionQuery.cpp
@@ -270,42 +270,24 @@ void FWorldConditionQuerySharedDefinition::Set(const TSubclassOf<UWorldCondition
 
 void FWorldConditionQuerySharedDefinition::PostSerialize(const FArchive& Ar)
 {
-	if (Conditions.Num() > 0)
+	const FUObjectSerializeContext* LoadContext = FUObjectThreadContext::Get().GetSerializeContext();
+	const UObject* SerializedObject = LoadContext ? LoadContext->SerializedObject : nullptr;
+	
+	if (!Link())
 	{
-		const FUObjectSerializeContext* LoadContext = FUObjectThreadContext::Get().GetSerializeContext();
-		const UObject* SerializedObject = LoadContext ? LoadContext->SerializedObject : nullptr;
-		
-		// Initialize conditions. This is done always on load so that any changes to the schema take affect.
-		if (SchemaClass)
-		{
-			const UWorldConditionSchema* Schema = SchemaClass.GetDefaultObject();
-			bool bResult = true;
-			for (int32 Index = 0; Index < Conditions.Num(); Index++)
-			{
-				FWorldConditionBase& Condition = Conditions[Index].GetMutable<FWorldConditionBase>();
-				bResult &= Condition.Initialize(*Schema);
-			}
-			if (!bResult)
-			{
-				UE_LOG(LogWorldCondition, Error, TEXT("World Condition: Failed to initialize query for %s."),
-					*GetNameSafe(SerializedObject));
-			}
-		}
-		else
- 
-		{
-			UE_LOG(LogWorldCondition, Error, TEXT("World Condition: shared definition for %s has empty schema, and %d conditions."),
-				*GetNameSafe(SerializedObject), Conditions.Num());
-		}
+		UE_LOG(LogWorldCondition, Error, TEXT("World Condition: Failed to link query for %s."),
+			*GetNameSafe(SerializedObject));
 	}
-
-	Link();
 }
 
 bool FWorldConditionQuerySharedDefinition::Link()
 {
 	bool bResult = true;
 
+	StateMinAlignment = 0;
+	StateSize = 0;
+	bIsLinked = false;
+
 	const UWorldConditionSchema* Schema = SchemaClass.GetDefaultObject();
 	if (!Schema)
 	{
@@ -360,8 +342,8 @@ bool FWorldConditionQuerySharedDefinition::Link()
 		}
 	}
 
-	StateMinAlignment = MinAlignment;
-	StateSize = Offset;
+	StateMinAlignment = uint8(MinAlignment);
+	StateSize = uint16(Offset);
 	
 	for (int32 Index = 0; Index < Conditions.Num(); Index++)
 	{
@@ -369,6 +351,8 @@ bool FWorldConditionQuerySharedDefinition::Link()
 		bResult &= Condition.Initialize(*Schema);
 	}
 
+	bIsLinked = bResult;
+
 	return bResult;
 }
 
@@ -378,7 +362,7 @@ bool FWorldConditionQuerySharedDefinition::Link()
 
 bool FWorldConditionQueryDefinition::IsValid() const
 {
-	return SharedDefinition.IsValid();
+	return SharedDefinition.IsValid() && SharedDefinition->IsLinked();
 }
 
 void FWorldConditionQueryDefinition::SetSchemaClass(const TSubclassOf<UWorldConditionSchema> InSchema)
@@ -509,6 +493,8 @@ bool FWorldConditionQueryDefinition::Serialize(FArchive& Ar)
 		{
 			UScriptStruct* Struct = TBaseStructure<FWorldConditionQuerySharedDefinition>::Get();
 			Struct->SerializeTaggedProperties(Ar, (uint8*)SharedDefinition.Get(), Struct, nullptr);
+			// SerializeTaggedProperties does not call PostSerialize() on the struct it's called (calls in items), call it manually.
+			SharedDefinition->PostSerialize(Ar);
 		}
 	}
 #if WITH_EDITOR
diff --git a/Engine/Plugins/Runtime/WorldConditions/Source/WorldConditions/Public/WorldConditionQuery.h b/Engine/Plugins/Runtime/WorldConditions/Source/WorldConditions/Public/WorldConditionQuery.h
--- a/Engine/Plugins/Runtime/WorldConditions/Source/WorldConditions/Public/WorldConditionQuery.h
+++ b/Engine/Plugins/Runtime/WorldConditions/Source/WorldConditions/Public/WorldConditionQuery.h
@@ -87,11 +87,12 @@ struct WORLDCONDITIONS_API FWorldConditionEditable
 USTRUCT()
 struct WORLDCONDITIONS_API FWorldConditionQuerySharedDefinition
 {
-private:
-	struct FPrivateToken { explicit FPrivateToken() = default; };
-	
-public:
 	GENERATED_BODY()
+	
+	FWorldConditionQuerySharedDefinition()
+		: bIsLinked(false)
+	{
+	}
 
 	void PostSerialize(const FArchive& Ar);
 
@@ -107,6 +108,9 @@ public:
 	/** @return the alignment of the memory needed for the state (in bytes). */
 	int32 GetStateMinAlignment() const { return StateMinAlignment; }
 
+	/** @return true if the shared definition is successfully linked. */
+	bool IsLinked() const { return bIsLinked; }
+	
 private:
 
 	/** Updates the state memory layout and initializes the conditions. */
@@ -129,6 +133,9 @@ private:
 	
 	/** Size of the state storage */
 	int32 StateSize = 0;
+
+	/** Flag indicating the last result of Link(), if true, the definition is ready to use. */
+	uint8 bIsLinked : 1;
 	
 	friend struct FWorldConditionQueryDefinition;
 };
diff --git a/Engine/Plugins/Runtime/WorldConditions/Source/WorldConditionsTestSuite/Private/WorldConditionsTest.cpp b/Engine/Plugins/Runtime/WorldConditions/Source/WorldConditionsTestSuite/Private/WorldConditionsTest.cpp
--- a/Engine/Plugins/Runtime/WorldConditions/Source/WorldConditionsTestSuite/Private/WorldConditionsTest.cpp
+++ b/Engine/Plugins/Runtime/WorldConditions/Source/WorldConditionsTestSuite/Private/WorldConditionsTest.cpp
@@ -3,6 +3,9 @@
 #include "AITestsCommon.h"
 #include "Engine/World.h"
 #include "WorldConditionTestTypes.h"
+#include "Serialization/MemoryReader.h"
+#include "Serialization/MemoryWriter.h"
+#include "Serialization/ObjectAndNameAsStringProxyArchive.h"
 
 #define LOCTEXT_NAMESPACE "AITestSuite_WorldConditionsTest"
 
@@ -266,6 +269,39 @@ struct FWorldConditionTest_FailingActivate : FAITestBase
 };
 IMPLEMENT_AI_INSTANT_TEST(FWorldConditionTest_FailingActivate, "System.WorldConditions.FailingActivate");
 
+struct FWorldConditionTest_Serialization : FAITestBase
+{
+	virtual bool InstantTest() override
+	{
+		UWorldConditionOwnerClass* Owner = NewAutoDestroyObject<UWorldConditionOwnerClass>();
+		const bool bInitialized = Owner->Definition.Initialize(GetWorld(), UWorldConditionTestSchema::StaticClass(),
+			{
+				FWorldConditionEditable(0, EWorldConditionOperator::Copy, FConstStructView::Make(FWorldConditionTest())),
+				FWorldConditionEditable(0, EWorldConditionOperator::And, FConstStructView::Make(FWorldConditionTest())),
+			});
+
+		AITEST_TRUE("Query definition should get initialized", bInitialized);
+		AITEST_EQUAL("Query successfully initialized definition should be valid", bInitialized, Owner->Definition.IsValid());
+
+		TArray<uint8> Data;
+		FMemoryWriter Writer(Data);
+		FObjectAndNameAsStringProxyArchive WriterProxy(Writer, /*bInLoadIfFindFails*/true);
+		Owner->Serialize(WriterProxy);
+
+		// Read back on new object
+		Owner = NewAutoDestroyObject<UWorldConditionOwnerClass>();
+		
+		FMemoryReader Reader(Data);
+		FObjectAndNameAsStringProxyArchive ReaderProxy(Reader, /*bInLoadIfFindFails*/true);
+		Owner->Serialize(ReaderProxy);
+
+		AITEST_TRUE("Query definition should be initialized after loading", Owner->Definition.IsValid());
+		
+		return true;
+	}
+};
+IMPLEMENT_AI_INSTANT_TEST(FWorldConditionTest_Serialization, "System.WorldConditions.Serialization");
+
 UE_ENABLE_OPTIMIZATION_SHIP
 
 #undef LOCTEXT_NAMESPACE
diff --git a/Engine/Plugins/Runtime/WorldConditions/Source/WorldConditionsTestSuite/Public/WorldConditionTestTypes.h b/Engine/Plugins/Runtime/WorldConditions/Source/WorldConditionsTestSuite/Public/WorldConditionTestTypes.h
--- a/Engine/Plugins/Runtime/WorldConditions/Source/WorldConditionsTestSuite/Public/WorldConditionTestTypes.h
+++ b/Engine/Plugins/Runtime/WorldConditions/Source/WorldConditionsTestSuite/Public/WorldConditionTestTypes.h
@@ -214,3 +214,14 @@ protected:
 	
 	int32 TestValue = 0;
 };
+
+UCLASS(hidden)
+class UWorldConditionOwnerClass : public UObject
+{
+	GENERATED_BODY()
+
+public:
+	
+	UPROPERTY(EditAnywhere, Category = "Default")
+	FWorldConditionQueryDefinition Definition;
+};
\ No newline at end of file
