Last relative include paths from module parent directory, will fix property in ue5-main to prevent conflicts

#rnx

diff --git a/Engine/Plugins/Runtime/WorldConditions/Source/WorldConditionsEditor/Private/Customizations/WorldConditionContextDataRefDetails.cpp b/Engine/Plugins/Runtime/WorldConditions/Source/WorldConditionsEditor/Private/Customizations/WorldConditionContextDataRefDetails.cpp
--- a/Engine/Plugins/Runtime/WorldConditions/Source/WorldConditionsEditor/Private/Customizations/WorldConditionContextDataRefDetails.cpp
+++ b/Engine/Plugins/Runtime/WorldConditions/Source/WorldConditionsEditor/Private/Customizations/WorldConditionContextDataRefDetails.cpp
@@ -9,7 +9,7 @@
 #include "WorldConditionQuery.h"
 #include "WorldConditionSchema.h"
 #include "Editor.h"
-#include "PropertyEditor/Private/PropertyNode.h"
+#include "Editor/PropertyEditor/Private/PropertyNode.h"
 
 #define LOCTEXT_NAMESPACE "WorldCondition"
 
