Add WorldCondition to know in which Biome is an Actor

diff --git a/Engine/Plugins/Runtime/WorldConditions/Source/WorldConditions/Public/WorldConditionTypes.h b/Engine/Plugins/Runtime/WorldConditions/Source/WorldConditions/Public/WorldConditionTypes.h
--- a/Engine/Plugins/Runtime/WorldConditions/Source/WorldConditions/Public/WorldConditionTypes.h
+++ b/Engine/Plugins/Runtime/WorldConditions/Source/WorldConditions/Public/WorldConditionTypes.h
@@ -82,6 +82,15 @@ namespace UE::WorldCondition
 		return ((InResult == EWorldConditionResultValue::IsTrue) ^ bInvert) ? EWorldConditionResultValue::IsTrue : EWorldConditionResultValue::IsFalse;
 	}
 
+	inline EWorldConditionResultValue FromOptional(const TOptional<bool>& InOptionalBool)
+	{
+		if (!InOptionalBool.IsSet())
+		{
+			return EWorldConditionResultValue::Invalid;
+		}
+		return InOptionalBool.GetValue() ? EWorldConditionResultValue::IsTrue : EWorldConditionResultValue::IsFalse;
+	}
+
 } // UE::WorldCondition
 
 /**
