World Condition: Fail gracefully when shared data is bad.

diff --git a/Engine/Plugins/Runtime/WorldConditions/Source/WorldConditions/Private/WorldConditionQuery.cpp b/Engine/Plugins/Runtime/WorldConditions/Source/WorldConditions/Private/WorldConditionQuery.cpp
--- a/Engine/Plugins/Runtime/WorldConditions/Source/WorldConditions/Private/WorldConditionQuery.cpp
+++ b/Engine/Plugins/Runtime/WorldConditions/Source/WorldConditions/Private/WorldConditionQuery.cpp
@@ -273,14 +273,14 @@ void FWorldConditionQuerySharedDefinition::PostSerialize(const FArchive& Ar)
 	const FUObjectSerializeContext* LoadContext = FUObjectThreadContext::Get().GetSerializeContext();
 	const UObject* SerializedObject = LoadContext ? LoadContext->SerializedObject : nullptr;
 	
-	if (!Link())
+	if (!Link(SerializedObject))
 	{
 		UE_LOG(LogWorldCondition, Error, TEXT("World Condition: Failed to link query for %s."),
 			*GetNameSafe(SerializedObject));
 	}
 }
 
-bool FWorldConditionQuerySharedDefinition::Link()
+bool FWorldConditionQuerySharedDefinition::Link(const UObject* Outer)
 {
 	bool bResult = true;
 
@@ -309,8 +309,15 @@ bool FWorldConditionQuerySharedDefinition::Link()
 	// Reserve space for all runtime data.
 	for (int32 Index = 0; Index < Conditions.Num(); Index++)
 	{
-		FWorldConditionBase& Condition = Conditions[Index].Get<FWorldConditionBase>();
-		if (const UStruct* StateStruct = Condition.GetRuntimeStateType())
+		FWorldConditionBase* Condition = Conditions[Index].GetPtr<FWorldConditionBase>();
+		if (!Condition)
+		{
+			UE_LOG(LogWorldCondition, Error, TEXT("World Condition: Encountered empty condition in a query definition for %s. This may be caused by the query using a struct from a plugin that is not loaded, or the condition type is removed."),
+				*GetNameSafe(Outer));
+			return false;
+		}
+		
+		if (const UStruct* StateStruct = Condition->GetRuntimeStateType())
 		{
 			int32 StructMinAlignment = 0;
 			int32 StructSize = 0;
@@ -319,26 +326,26 @@ bool FWorldConditionQuerySharedDefinition::Link()
 			{
 				StructMinAlignment = ScriptStruct->GetMinAlignment();
 				StructSize = ScriptStruct->GetStructureSize();
-				Condition.bIsStateObject = false;
+				Condition->bIsStateObject = false;
 			}
 			else if (const UClass* Class = Cast<const UClass>(StateStruct))
 			{
 				StructMinAlignment = FWorldConditionStateObject::StaticStruct()->GetMinAlignment();
 				StructSize = FWorldConditionStateObject::StaticStruct()->GetStructureSize();
-				Condition.bIsStateObject = true;
+				Condition->bIsStateObject = true;
 			}
 			
 			check(StructMinAlignment > 0 && StructSize > 0);
 
 			Offset = Align(Offset, StructMinAlignment);
-			Condition.StateDataOffset = IntCastChecked<uint16>(Offset);
+			Condition->StateDataOffset = IntCastChecked<uint16>(Offset);
 
 			Offset += StructSize;
 			MinAlignment = FMath::Max(MinAlignment, StructMinAlignment);
 		}
 		else
 		{
-			Condition.StateDataOffset = 0;
+			Condition->StateDataOffset = 0;
 		}
 	}
 
@@ -452,7 +459,7 @@ bool FWorldConditionQueryDefinition::Initialize(const UObject* Outer)
 	// Create a new shared definition to allow the allocated states to deactivate properly event if we update the definition.
 	SharedDefinition = MakeShared<FWorldConditionQuerySharedDefinition>();
 	SharedDefinition->Set(SchemaClass, ValidConditions);
-	SharedDefinition->Link();
+	SharedDefinition->Link(Outer);
 
 #endif
 	
diff --git a/Engine/Plugins/Runtime/WorldConditions/Source/WorldConditions/Public/WorldConditionQuery.h b/Engine/Plugins/Runtime/WorldConditions/Source/WorldConditions/Public/WorldConditionQuery.h
--- a/Engine/Plugins/Runtime/WorldConditions/Source/WorldConditions/Public/WorldConditionQuery.h
+++ b/Engine/Plugins/Runtime/WorldConditions/Source/WorldConditions/Public/WorldConditionQuery.h
@@ -114,7 +114,7 @@ struct WORLDCONDITIONS_API FWorldConditionQuerySharedDefinition
 private:
 
 	/** Updates the state memory layout and initializes the conditions. */
-	bool Link();
+	bool Link(const UObject* Outer);
 	
 	/** Sets the definition data. */
 	void Set(const TSubclassOf<UWorldConditionSchema> InSchema, const TArrayView<FConstStructView> InConditions);
