World Condition: Use shared ptr for state data and weak ptr for invalidation handles to avoid crashes.

diff --git a/Engine/Plugins/Runtime/WorldConditions/Source/WorldConditions/Private/WorldConditionQuery.cpp b/Engine/Plugins/Runtime/WorldConditions/Source/WorldConditions/Private/WorldConditionQuery.cpp
--- a/Engine/Plugins/Runtime/WorldConditions/Source/WorldConditions/Private/WorldConditionQuery.cpp
+++ b/Engine/Plugins/Runtime/WorldConditions/Source/WorldConditions/Private/WorldConditionQuery.cpp
@@ -37,11 +37,18 @@ FCustomVersionRegistration GWorldConditionCustomVersion(FWorldConditionCustomVer
 
 void FWorldConditionResultInvalidationHandle::InvalidateResult() const
 {
-	if (CachedResult && Item)
+	const TSharedPtr<uint8> StateMemory = WeakStateMemory.Pin();
+	if (StateMemory.IsValid())
 	{
+		EWorldConditionResultValue* CachedResult = reinterpret_cast<EWorldConditionResultValue*>(StateMemory.Get() + FWorldConditionQueryState::CachedResultOffset);
+		FWorldConditionItem* Item = reinterpret_cast<FWorldConditionItem*>(StateMemory.Get() + ItemOffset);
 		*CachedResult = EWorldConditionResultValue::Invalid;
 		Item->CachedResult = EWorldConditionResultValue::Invalid;
 	}
+	else
+	{
+		UE_LOG(LogWorldCondition, Warning, TEXT("World Condition: Trying to access freed state memory while calling InvalidateResult(). Make sure you are calling Deactivate() on your world condition query, and that you unregister any delegates on world conditions."));
+	}
 }
 
 
@@ -71,11 +78,10 @@ FWorldConditionQueryState::~FWorldConditionQueryState()
 		}
 	}
 
-	if (Memory)
+	if (Memory.IsValid())
 	{
 		// Something went very wrong, this should never happen.
-		UE_LOG(LogWorldCondition, Error, TEXT("World Condition: State %p has still allocated memory in destructor, might leak memory."), Memory);
-		FMemory::Free(Memory);
+		UE_LOG(LogWorldCondition, Error, TEXT("World Condition: State %p has still allocated memory in destructor, might leak memory."), Memory.Get());
 		Memory = nullptr;
 	}
 }
@@ -110,17 +116,22 @@ void FWorldConditionQueryState::InitializeInternal(const UObject* InOwner, const
 
 	// Cache num conditions so that we can access the items and cached data without touching the definition.
 	NumConditions = IntCastChecked<uint8>(Conditions.Num());
-	
-	Memory = static_cast<uint8*>(FMemory::Malloc(SharedDefinition->GetStateSize(), SharedDefinition->GetStateMinAlignment()));
+
+	uint8* AllocatedMemory = static_cast<uint8*>(FMemory::Malloc(SharedDefinition->GetStateSize(), SharedDefinition->GetStateMinAlignment()));
+	Memory = TSharedPtr<uint8>(AllocatedMemory,
+		/* Deleter */[](uint8* Obj)
+		{
+			FMemory::Free(Obj);
+		});
 
 	// Init cached result
-	EWorldConditionResultValue& CachedResult = *reinterpret_cast<EWorldConditionResultValue*>(Memory + CachedResultOffset);
+	EWorldConditionResultValue& CachedResult = *reinterpret_cast<EWorldConditionResultValue*>(AllocatedMemory + CachedResultOffset);
 	CachedResult = EWorldConditionResultValue::Invalid;
 	
 	// Initialize items
 	for (int32 Index = 0; Index < static_cast<int32>(NumConditions); Index++)
 	{
-		new (Memory + ItemsOffset + sizeof(FWorldConditionItem) * Index) FWorldConditionItem();
+		new (AllocatedMemory + ItemsOffset + sizeof(FWorldConditionItem) * Index) FWorldConditionItem();
 	}
 
 	// Initialize state
@@ -131,7 +142,7 @@ void FWorldConditionQueryState::InitializeInternal(const UObject* InOwner, const
 		{
 			continue;
 		}
-		uint8* StateMemory = Memory + Condition.GetStateDataOffset();
+		uint8* StateMemory = AllocatedMemory + Condition.GetStateDataOffset();
 		if (Condition.IsStateObject())
 		{
 			new (StateMemory) FWorldConditionStateObject();
@@ -151,7 +162,7 @@ void FWorldConditionQueryState::InitializeInternal(const UObject* InOwner, const
 
 void FWorldConditionQueryState::Free()
 {
-	if (Memory == nullptr)
+	if (!Memory.IsValid())
 	{
 		NumConditions = 0;
 		SharedDefinition = nullptr;
@@ -172,10 +183,10 @@ void FWorldConditionQueryState::Free()
 		{
 			continue;
 		}
-		uint8* StateMemory = Memory + Condition.GetStateDataOffset();
+		uint8* StateMemory = Memory.Get() + Condition.GetStateDataOffset();
 		if (Condition.IsStateObject())
 		{
-			FWorldConditionStateObject& StateObject = *reinterpret_cast<FWorldConditionStateObject*>(Memory + Condition.GetStateDataOffset());
+			FWorldConditionStateObject& StateObject = *reinterpret_cast<FWorldConditionStateObject*>(Memory.Get() + Condition.GetStateDataOffset());
 			StateObject.~FWorldConditionStateObject();
 		}
 		else
@@ -185,8 +196,6 @@ void FWorldConditionQueryState::Free()
 		}
 	}
 
-	FMemory::Free(Memory);
-	
 	Memory = nullptr;
 	NumConditions = 0;
 	SharedDefinition = nullptr;
@@ -202,7 +211,7 @@ void FWorldConditionQueryState::AddStructReferencedObjects(class FReferenceColle
 		Collector.AddPropertyReferencesWithStructARO(FWorldConditionQuerySharedDefinition::StaticStruct(), SharedDefinition.Get());
 	}
 	
-	if (Memory == nullptr || !SharedDefinition.IsValid())
+	if (!Memory.IsValid() || !SharedDefinition.IsValid())
 	{
 		return;
 	}
@@ -219,7 +228,7 @@ void FWorldConditionQueryState::AddStructReferencedObjects(class FReferenceColle
 			continue;
 		}
 		
-		uint8* StateMemory = Memory + Condition.GetStateDataOffset();
+		uint8* StateMemory = Memory.Get() + Condition.GetStateDataOffset();
 		if (Condition.IsStateObject())
 		{
 			FWorldConditionStateObject& StateObject = *reinterpret_cast<FWorldConditionStateObject*>(StateMemory);
@@ -239,12 +248,9 @@ void FWorldConditionQueryState::AddStructReferencedObjects(class FReferenceColle
 FWorldConditionResultInvalidationHandle FWorldConditionQueryState::GetInvalidationHandle(const FWorldConditionBase& Condition) const
 {
 	check(bIsInitialized);
-	check(Memory && Condition.GetConditionIndex() < NumConditions);
-
-	EWorldConditionResultValue* CachedResult = reinterpret_cast<EWorldConditionResultValue*>(Memory + CachedResultOffset);
-	FWorldConditionItem* Item = reinterpret_cast<FWorldConditionItem*>(Memory + ItemsOffset + Condition.GetConditionIndex() * sizeof(FWorldConditionItem));
-
-	return FWorldConditionResultInvalidationHandle(CachedResult, Item);
+	check(Memory.IsValid() && Condition.GetConditionIndex() < NumConditions);
+	const int32 ItemOffset = ItemsOffset + Condition.GetConditionIndex() * sizeof(FWorldConditionItem);
+	return FWorldConditionResultInvalidationHandle(Memory, ItemOffset);
 }
 
 
diff --git a/Engine/Plugins/Runtime/WorldConditions/Source/WorldConditions/Public/WorldConditionQuery.h b/Engine/Plugins/Runtime/WorldConditions/Source/WorldConditions/Public/WorldConditionQuery.h
--- a/Engine/Plugins/Runtime/WorldConditions/Source/WorldConditions/Public/WorldConditionQuery.h
+++ b/Engine/Plugins/Runtime/WorldConditions/Source/WorldConditions/Public/WorldConditionQuery.h
@@ -272,13 +272,13 @@ struct WORLDCONDITIONS_API FWorldConditionResultInvalidationHandle
 	FWorldConditionResultInvalidationHandle() = default;
 
 	/** @return True if the handle is valid. */
-	bool IsValid() const { return CachedResult != nullptr && Item != nullptr; }
+	bool IsValid() const { return WeakStateMemory.IsValid(); }
 
 	/** Resets the handle. */
 	void Reset()
 	{
-		CachedResult = nullptr;
-		Item = nullptr;
+		WeakStateMemory = nullptr;
+		ItemOffset = 0;
 	}
 
 	/**
@@ -288,14 +288,14 @@ struct WORLDCONDITIONS_API FWorldConditionResultInvalidationHandle
 	
 protected:
 
-	explicit FWorldConditionResultInvalidationHandle(EWorldConditionResultValue* InCachedResult, FWorldConditionItem* InItem)
-		: CachedResult(InCachedResult)
-		, Item(InItem)
+	explicit FWorldConditionResultInvalidationHandle(const TWeakPtr<uint8> InStateMemory, const int32 InItemOffset)
+		: WeakStateMemory(InStateMemory)
+		, ItemOffset(InItemOffset)
 	{
 	}
 
-	EWorldConditionResultValue* CachedResult = nullptr;
-	FWorldConditionItem* Item = nullptr;
+	TWeakPtr<uint8> WeakStateMemory = nullptr;
+	int32 ItemOffset = 0;
 
 	friend struct FWorldConditionQueryState;
 };
@@ -353,15 +353,15 @@ struct WORLDCONDITIONS_API FWorldConditionQueryState
 	{
 		check(bIsInitialized);
 		// If Memory is null, and the query is initialized, it's empty query, which evaluates to true. 
-		return Memory ? *reinterpret_cast<EWorldConditionResultValue*>(Memory + CachedResultOffset) : EWorldConditionResultValue::IsTrue;
+		return Memory.IsValid() ? *reinterpret_cast<EWorldConditionResultValue*>(Memory.Get() + CachedResultOffset) : EWorldConditionResultValue::IsTrue;
 	}
 
 	/** @return Condition item at specific index */
 	FWorldConditionItem& GetItem(const int32 Index) const
 	{
 		check(bIsInitialized);
-		check(Memory && Index >= 0 && Index < (int32)NumConditions);
-		return *reinterpret_cast<FWorldConditionItem*>(Memory + ItemsOffset + Index * sizeof(FWorldConditionItem));
+		check(Memory.IsValid() && Index >= 0 && Index < (int32)NumConditions);
+		return *reinterpret_cast<FWorldConditionItem*>(Memory.Get() + ItemsOffset + Index * sizeof(FWorldConditionItem));
 	}
 
 	/** @return Object describing the state of a specified condition. */
@@ -370,7 +370,8 @@ struct WORLDCONDITIONS_API FWorldConditionQueryState
 		check(bIsInitialized);
 		check(Condition.GetStateDataOffset() > 0);
 		check(Condition.IsStateObject());
-		const FWorldConditionStateObject& StateObject = *reinterpret_cast<FWorldConditionStateObject*>(Memory + Condition.GetStateDataOffset());
+		check(Memory.IsValid());
+		const FWorldConditionStateObject& StateObject = *reinterpret_cast<FWorldConditionStateObject*>(Memory.Get() + Condition.GetStateDataOffset());
 		return StateObject.Object;
 	}
 
@@ -382,7 +383,8 @@ struct WORLDCONDITIONS_API FWorldConditionQueryState
 		check (!Condition.IsStateObject());
 		const UScriptStruct* ScriptStruct = Cast<UScriptStruct>(Condition.GetRuntimeStateType());
 		check(ScriptStruct);
-		return FStructView(Cast<UScriptStruct>(Condition.GetRuntimeStateType()), Memory + Condition.GetStateDataOffset());
+		check(Memory.IsValid());
+		return FStructView(Cast<UScriptStruct>(Condition.GetRuntimeStateType()), Memory.Get() + Condition.GetStateDataOffset());
 	}
 
 	/** @return The number of conditions in the state data. */
@@ -409,9 +411,9 @@ struct WORLDCONDITIONS_API FWorldConditionQueryState
 	void SetCachedResult(const EWorldConditionResultValue InResult) const
 	{
 		check(bIsInitialized);
-		if (Memory)
+		if (Memory.IsValid())
 		{
-			EWorldConditionResultValue& CachedResult = *reinterpret_cast<EWorldConditionResultValue*>(Memory + CachedResultOffset); 
+			EWorldConditionResultValue& CachedResult = *reinterpret_cast<EWorldConditionResultValue*>(Memory.Get() + CachedResultOffset); 
 			CachedResult = InResult;
 		}
 	}
@@ -428,8 +430,7 @@ private:
 
 	uint8 NumConditions = 0;
 	uint8 bIsInitialized : 1;
-	uint8* Memory = nullptr;
-
+	TSharedPtr<uint8> Memory = nullptr;
 	TSharedPtr<FWorldConditionQuerySharedDefinition> SharedDefinition = nullptr;
 
 	UPROPERTY(Transient)
