World Conditions: Added pretty printing description.

diff --git a/Engine/Plugins/Runtime/WorldConditions/Source/WorldConditions/Private/WorldConditionQuery.cpp b/Engine/Plugins/Runtime/WorldConditions/Source/WorldConditions/Private/WorldConditionQuery.cpp
--- a/Engine/Plugins/Runtime/WorldConditions/Source/WorldConditions/Private/WorldConditionQuery.cpp
+++ b/Engine/Plugins/Runtime/WorldConditions/Source/WorldConditions/Private/WorldConditionQuery.cpp
@@ -4,9 +4,12 @@
 #include "WorldConditionContext.h"
 #include "UObject/UObjectThreadContext.h"
 #include "Serialization/CustomVersion.h"
+#include "Misc/StringBuilder.h"
 
 #include UE_INLINE_GENERATED_CPP_BY_NAME(WorldConditionQuery)
 
+#define LOCTEXT_NAMESPACE "WorldCondition"
+
 struct FWorldConditionCustomVersion
 {
 	enum Type
@@ -392,6 +395,84 @@ bool FWorldConditionQueryDefinition::Initialize(const UObject* Outer, const TSub
 }
 #endif
 
+#if WITH_EDITOR
+FText FWorldConditionQueryDefinition::GetDescription() const
+{
+	if (!SharedDefinition.IsValid())
+	{
+		return LOCTEXT("Empty", "Empty");
+	}
+
+	const FInstancedStructContainer& Conditions = SharedDefinition->GetConditions();
+	if (Conditions.Num() == 0)
+	{
+		return LOCTEXT("Empty", "Empty");
+	}
+
+	TStringBuilder<256> Builder;
+	
+	const FWorldConditionBase* PrevCondition = nullptr;
+	for (int32 Index = 0; Index < Conditions.Num(); Index++)
+	{
+		const FWorldConditionBase& Condition = Conditions[Index].Get<const FWorldConditionBase>();
+		const int32 CurrDepth = PrevCondition ? PrevCondition->GetNextExpressionDepth() : 0;
+		const int32 NextDepth = Condition.GetNextExpressionDepth();
+		const int32 DeltaDepth = NextDepth - CurrDepth;
+		const int32 OpenParens = FMath::Max(0, DeltaDepth);
+		const int32 CloseParens = FMath::Max(0, -DeltaDepth);
+
+		// Operator
+		FText Operator;
+		if (Index == 0)
+		{
+			Operator = LOCTEXT("IfOperator", "IF");
+		}
+		else if (Condition.GetOperator() == EWorldConditionOperator::And)
+		{
+			Operator = LOCTEXT("AndOperator", "AND");
+		}
+		else if (Condition.GetOperator() == EWorldConditionOperator::Or)
+		{
+			Operator = LOCTEXT("OrOperator", "OR");
+		}
+		else
+		{
+			ensureMsgf(TEXT("Unhandled operator %s\n"), *UEnum::GetValueAsString(Condition.GetOperator()));
+			break;
+		}
+		Builder.Append(Operator.ToString());
+		Builder.AppendChar(' ');
+
+		// Open parens
+		for (int32 Paren = 0; Paren < OpenParens; Paren++)
+		{
+			Builder.AppendChar('(');
+		}
+
+		// Item desc
+		FText ConditionDesc = Condition.GetDescription();
+		Builder.AppendChar('[');
+		Builder.Append(ConditionDesc.ToString());
+		Builder.AppendChar(']');
+
+		// Close parens
+		for (int32 Paren = 0; Paren < CloseParens; Paren++)
+		{
+			Builder.AppendChar(')');
+		}
+
+		if ((Index + 1) < Conditions.Num())
+		{
+			Builder.AppendChar(' ');
+		}
+
+		PrevCondition = &Condition;
+	}
+
+	return FText::FromString(Builder.ToString());
+}
+#endif
+
 bool FWorldConditionQueryDefinition::Initialize(const UObject* Outer)
 {
 	bool bResult = true;
@@ -540,6 +621,20 @@ bool FWorldConditionQueryDefinition::ImportTextItem(const TCHAR*& Buffer, int32
 	return false;
 }
 
+bool FWorldConditionQueryDefinition::ExportTextItem(FString& ValueStr, FWorldConditionQueryDefinition const& DefaultValue, class UObject* Parent, int32 PortFlags, class UObject* ExportRootScope) const
+{
+#if WITH_EDITOR
+	if (PortFlags & PPF_PropertyWindow)
+	{
+		const FText Desc = GetDescription();
+		ValueStr = Desc.ToString();
+		return true;
+	}
+#endif
+	
+	return false;
+}
+
 void FWorldConditionQueryDefinition::AddStructReferencedObjects(FReferenceCollector& Collector) const
 {
 	if (SharedDefinition.IsValid())
@@ -593,3 +688,5 @@ bool FWorldConditionQuery::IsActive() const
 {
 	return QueryState.IsInitialized();
 }
+
+#undef LOCTEXT_NAMESPACE
\ No newline at end of file
diff --git a/Engine/Plugins/Runtime/WorldConditions/Source/WorldConditions/Public/WorldConditionQuery.h b/Engine/Plugins/Runtime/WorldConditions/Source/WorldConditions/Public/WorldConditionQuery.h
--- a/Engine/Plugins/Runtime/WorldConditions/Source/WorldConditions/Public/WorldConditionQuery.h
+++ b/Engine/Plugins/Runtime/WorldConditions/Source/WorldConditions/Public/WorldConditionQuery.h
@@ -177,6 +177,7 @@ struct WORLDCONDITIONS_API FWorldConditionQueryDefinition
 
 	bool Serialize(FArchive& Ar);
 	bool ImportTextItem(const TCHAR*& Buffer, int32 PortFlags, UObject* Parent, FOutputDevice* ErrorText, FArchive* InSerializingArchive = nullptr);
+	bool ExportTextItem(FString& ValueStr, FWorldConditionQueryDefinition const& DefaultValue, class UObject* Parent, int32 PortFlags, class UObject* ExportRootScope) const;
 	void AddStructReferencedObjects(FReferenceCollector& Collector) const;
 
 
@@ -191,6 +192,11 @@ struct WORLDCONDITIONS_API FWorldConditionQueryDefinition
 	}
 #endif
 
+#if WITH_EDITOR
+	/** @returs description of the query expression. */
+	FText GetDescription() const;
+#endif
+
 	const TSharedPtr<FWorldConditionQuerySharedDefinition> GetSharedDefinitionPtr() const { return SharedDefinition; }
 	
 private:
@@ -219,6 +225,7 @@ struct TStructOpsTypeTraits<FWorldConditionQueryDefinition> : public TStructOpsT
 		WithAddStructReferencedObjects = true,
 		WithSerializer = true,
 		WithImportTextItem = true,
+		WithExportTextItem = true,
 	};
 };
 
diff --git a/Engine/Plugins/Runtime/WorldConditions/Source/WorldConditionsTestSuite/Private/WorldConditionsTest.cpp b/Engine/Plugins/Runtime/WorldConditions/Source/WorldConditionsTestSuite/Private/WorldConditionsTest.cpp
--- a/Engine/Plugins/Runtime/WorldConditions/Source/WorldConditionsTestSuite/Private/WorldConditionsTest.cpp
+++ b/Engine/Plugins/Runtime/WorldConditions/Source/WorldConditionsTestSuite/Private/WorldConditionsTest.cpp
@@ -302,6 +302,45 @@ struct FWorldConditionTest_Serialization : FAITestBase
 };
 IMPLEMENT_AI_INSTANT_TEST(FWorldConditionTest_Serialization, "System.WorldConditions.Serialization");
 
+struct FWorldConditionTest_Description : FAITestBase
+{
+	virtual bool InstantTest() override
+	{
+		TArray<FWorldConditionEditable> Conditions = {
+			FWorldConditionEditable(/*Depth*/0, EWorldConditionOperator::Copy, FConstStructView::Make(FWorldConditionTest(0))),	//	if	(A
+			FWorldConditionEditable(/*Depth*/1, EWorldConditionOperator::Or,   FConstStructView::Make(FWorldConditionTest(1))),	//	.	|| B)
+			FWorldConditionEditable(/*Depth*/0, EWorldConditionOperator::And,  FConstStructView::Make(FWorldConditionTest(1))),	//	&&	(	(C
+			FWorldConditionEditable(/*Depth*/2, EWorldConditionOperator::And,  FConstStructView::Make(FWorldConditionTest(1))),	//	.	.	&& D)
+			FWorldConditionEditable(/*Depth*/1, EWorldConditionOperator::Or,   FConstStructView::Make(FWorldConditionTest(0))),	//	.	|| E)
+		};
+		
+		FWorldConditionQueryDefinition Definition;
+		const bool bInitialized = Definition.Initialize(&GetWorld(), UWorldConditionTestSchema::StaticClass(), Conditions);
+		AITEST_TRUE("Query definition should get initialized", bInitialized);
+
+		auto GetDescriptionString = [](const FWorldConditionEditable& EditableCondition)
+		{
+			const FWorldConditionTest& Condition = EditableCondition.Condition.Get<const FWorldConditionTest>();
+			return Condition.GetDescription().ToString();
+		};
+		
+		const FString Expected =
+			TEXT("IF ([") + GetDescriptionString(Conditions[0])
+			+ TEXT("] OR [") + GetDescriptionString(Conditions[1])
+			+ TEXT("]) AND (([") + GetDescriptionString(Conditions[2])
+			+ TEXT("] AND [") + GetDescriptionString(Conditions[3])
+			+ TEXT("]) OR [") + GetDescriptionString(Conditions[4])
+			+ TEXT("])");
+
+		const FText Desc = Definition.GetDescription();
+
+		AITEST_EQUAL("Query text shold match expected result.", Desc.ToString(), Expected);
+		
+		return true;
+	}
+};
+IMPLEMENT_AI_INSTANT_TEST(FWorldConditionTest_Description, "System.WorldConditions.Description");
+
 UE_ENABLE_OPTIMIZATION_SHIP
 
 #undef LOCTEXT_NAMESPACE
diff --git a/Engine/Plugins/Runtime/WorldConditions/Source/WorldConditionsTestSuite/Public/WorldConditionTestTypes.h b/Engine/Plugins/Runtime/WorldConditions/Source/WorldConditionsTestSuite/Public/WorldConditionTestTypes.h
--- a/Engine/Plugins/Runtime/WorldConditions/Source/WorldConditionsTestSuite/Public/WorldConditionTestTypes.h
+++ b/Engine/Plugins/Runtime/WorldConditions/Source/WorldConditionsTestSuite/Public/WorldConditionTestTypes.h
@@ -112,9 +112,13 @@ struct FWorldConditionTest : public FWorldConditionBase
 	
 	virtual void Deactivate(const FWorldConditionContext& Context) const override
 	{
-		
 	}
 
+	virtual FText GetDescription() const override
+	{
+		return FText::Format(FText::FromString(TEXT("Value == {0}")), FText::AsNumber(TestValue));
+	}
+	
 protected:
 
 	FWorldConditionContextDataRef ValueRef;
