WorldConditions: Fixed world condition definitions not being copied properly on actor components
- Added Indentical() methods for world condition definition and shared definition
- Added WithCopy to world condition definition so that the shared pointer gets copied

diff --git a/Engine/Plugins/Runtime/WorldConditions/Source/WorldConditions/Private/WorldConditionQuery.cpp b/Engine/Plugins/Runtime/WorldConditions/Source/WorldConditions/Private/WorldConditionQuery.cpp
--- a/Engine/Plugins/Runtime/WorldConditions/Source/WorldConditions/Private/WorldConditionQuery.cpp
+++ b/Engine/Plugins/Runtime/WorldConditions/Source/WorldConditions/Private/WorldConditionQuery.cpp
@@ -289,6 +289,25 @@ void FWorldConditionQuerySharedDefinition::PostSerialize(const FArchive& Ar)
 	}
 }
 
+bool FWorldConditionQuerySharedDefinition::Identical(const FWorldConditionQuerySharedDefinition* Other, uint32 PortFlags) const
+{
+	if (!Other)
+	{
+		return false;
+	}
+
+	if (SchemaClass != Other->SchemaClass
+		|| StateMinAlignment != Other->StateMinAlignment
+		|| StateSize != Other->StateSize
+		|| bIsLinked != Other->bIsLinked)
+	{
+		return false;
+	}
+
+	return Conditions.Identical(&Other->Conditions, PortFlags);
+}
+
+
 bool FWorldConditionQuerySharedDefinition::Link(const UObject* Outer)
 {
 	bool bResult = true;
@@ -643,6 +662,36 @@ void FWorldConditionQueryDefinition::AddStructReferencedObjects(FReferenceCollec
 	}
 }
 
+bool FWorldConditionQueryDefinition::Identical(const FWorldConditionQueryDefinition* Other, uint32 PortFlags) const
+{
+	if (!Other)
+	{
+		return false;
+	}
+
+	if (SchemaClass != Other->SchemaClass
+		|| SharedDefinition.IsValid() != Other->SharedDefinition.IsValid())
+	{
+		return false;
+	}
+
+	if (SharedDefinition.IsValid()
+		&& Other->SharedDefinition.IsValid()
+		&& !SharedDefinition->Identical(Other->SharedDefinition.Get(), PortFlags))
+	{
+		return false;
+	}
+	
+#if WITH_EDITORONLY_DATA
+	if (EditableConditions != Other->EditableConditions)
+	{
+		return false;
+	}
+#endif
+
+	return true;
+}
+
 
 //
 // FWorldConditionQuery
diff --git a/Engine/Plugins/Runtime/WorldConditions/Source/WorldConditions/Public/WorldConditionQuery.h b/Engine/Plugins/Runtime/WorldConditions/Source/WorldConditions/Public/WorldConditionQuery.h
--- a/Engine/Plugins/Runtime/WorldConditions/Source/WorldConditions/Public/WorldConditionQuery.h
+++ b/Engine/Plugins/Runtime/WorldConditions/Source/WorldConditions/Public/WorldConditionQuery.h
@@ -61,7 +61,20 @@ struct WORLDCONDITIONS_API FWorldConditionEditable
 		ExpressionDepth = 0;
 		Condition.Reset();
 	}
-	
+
+	bool operator==(const FWorldConditionEditable& Other) const
+	{
+		return ExpressionDepth == Other.ExpressionDepth
+				&& Operator == Other.Operator
+				&& bInvert == Other.bInvert
+				&& Condition == Other.Condition;
+	}
+
+	bool operator!=(const FWorldConditionEditable& Other) const
+	{
+		return !(*this == Other);
+	}
+
 	/** Expression depth controlling the parenthesis of the expression. */
 	UPROPERTY(EditAnywhere, Category="Default")
 	uint8 ExpressionDepth = 0;
@@ -95,6 +108,7 @@ struct WORLDCONDITIONS_API FWorldConditionQuerySharedDefinition
 	}
 
 	void PostSerialize(const FArchive& Ar);
+	bool Identical(const FWorldConditionQuerySharedDefinition* Other, uint32 PortFlags) const;
 
 	/** @return the schema used for the definition. */
 	TSubclassOf<UWorldConditionSchema> GetSchemaClass() const { return SchemaClass; }
@@ -146,6 +160,7 @@ struct TStructOpsTypeTraits<FWorldConditionQuerySharedDefinition> : public TStru
 	enum
 	{
 		WithPostSerialize = true,
+		WithIdentical = true,
 	};
 };
 
@@ -179,7 +194,7 @@ struct WORLDCONDITIONS_API FWorldConditionQueryDefinition
 	bool ImportTextItem(const TCHAR*& Buffer, int32 PortFlags, UObject* Parent, FOutputDevice* ErrorText, FArchive* InSerializingArchive = nullptr);
 	bool ExportTextItem(FString& ValueStr, FWorldConditionQueryDefinition const& DefaultValue, class UObject* Parent, int32 PortFlags, class UObject* ExportRootScope) const;
 	void AddStructReferencedObjects(FReferenceCollector& Collector) const;
-
+	bool Identical(const FWorldConditionQueryDefinition* Other, uint32 PortFlags) const;
 
 #if WITH_EDITORONLY_DATA
 	/** Initialized the condition with specific data. */
@@ -222,6 +237,8 @@ struct TStructOpsTypeTraits<FWorldConditionQueryDefinition> : public TStructOpsT
 {
 	enum
 	{
+		WithCopy = true, // Ensures that the SharedDefinition gets copied too.
+		WithIdentical = true,
 		WithAddStructReferencedObjects = true,
 		WithSerializer = true,
 		WithImportTextItem = true,
