WorldCondition: Do not initialize world condition state from defintion that has failed to link.

diff --git a/Engine/Plugins/Runtime/WorldConditions/Source/WorldConditions/Private/WorldConditionQuery.cpp b/Engine/Plugins/Runtime/WorldConditions/Source/WorldConditions/Private/WorldConditionQuery.cpp
--- a/Engine/Plugins/Runtime/WorldConditions/Source/WorldConditions/Private/WorldConditionQuery.cpp
+++ b/Engine/Plugins/Runtime/WorldConditions/Source/WorldConditions/Private/WorldConditionQuery.cpp
@@ -113,6 +113,12 @@ void FWorldConditionQueryState::InitializeInternal(const UObject* InOwner, const
 		return;
 	}
 
+	if (!InSharedDefinition->IsLinked())
+	{
+		UE_LOG(LogWorldCondition, Error, TEXT("World Condition: Trying to initialize query state with invalid definition for %s."), *GetNameSafe(Owner));
+		return;
+	}
+	
 	SharedDefinition = InSharedDefinition;
 
 	const FInstancedStructContainer& Conditions = SharedDefinition->GetConditions(); 
