[WorldCondition] fixed condition index not set in FWorldConditionBase when linking the shared definition. This fixes issues with the Invalidation handles.

diff --git a/Engine/Plugins/Runtime/WorldConditions/Source/WorldConditions/Private/WorldConditionQuery.cpp b/Engine/Plugins/Runtime/WorldConditions/Source/WorldConditions/Private/WorldConditionQuery.cpp
--- a/Engine/Plugins/Runtime/WorldConditions/Source/WorldConditions/Private/WorldConditionQuery.cpp
+++ b/Engine/Plugins/Runtime/WorldConditions/Source/WorldConditions/Private/WorldConditionQuery.cpp
@@ -353,7 +353,8 @@ bool FWorldConditionQuerySharedDefinition::Link(const UObject* Outer)
 				*GetNameSafe(Outer));
 			return false;
 		}
-		
+		Condition->ConditionIndex = Index;
+
 		if (auto* StateStruct = Condition->GetRuntimeStateType(); StateStruct && *StateStruct)
 		{
 			int32 StructMinAlignment = 0;
