[WorldConditionsEditor] removed unused forward declarations

diff --git a/Engine/Plugins/Runtime/WorldConditions/Source/WorldConditionsEditor/Private/Customizations/WorldConditionContextDataRefDetails.h b/Engine/Plugins/Runtime/WorldConditions/Source/WorldConditionsEditor/Private/Customizations/WorldConditionContextDataRefDetails.h
--- a/Engine/Plugins/Runtime/WorldConditions/Source/WorldConditionsEditor/Private/Customizations/WorldConditionContextDataRefDetails.h
+++ b/Engine/Plugins/Runtime/WorldConditions/Source/WorldConditionsEditor/Private/Customizations/WorldConditionContextDataRefDetails.h
@@ -6,7 +6,6 @@
 
 class IPropertyHandle;
 class IPropertyUtilities;
-class USmartObjectDefinition;
 class SWidget;
 struct FWorldConditionContextDataRef;
 class UWorldConditionSchema;
diff --git a/Engine/Plugins/Runtime/WorldConditions/Source/WorldConditionsEditor/Private/Customizations/WorldConditionQueryDefinitionDetails.h b/Engine/Plugins/Runtime/WorldConditions/Source/WorldConditionsEditor/Private/Customizations/WorldConditionQueryDefinitionDetails.h
--- a/Engine/Plugins/Runtime/WorldConditions/Source/WorldConditionsEditor/Private/Customizations/WorldConditionQueryDefinitionDetails.h
+++ b/Engine/Plugins/Runtime/WorldConditions/Source/WorldConditionsEditor/Private/Customizations/WorldConditionQueryDefinitionDetails.h
@@ -6,7 +6,6 @@
 
 class IPropertyHandle;
 class IPropertyUtilities;
-class USmartObjectDefinition;
 class SWidget;
 
 /**
