[WorldCondition] Keep track of the world condition activation status

diff --git a/Engine/Plugins/Runtime/WorldConditions/Source/WorldConditions/Private/WorldConditionContext.cpp b/Engine/Plugins/Runtime/WorldConditions/Source/WorldConditions/Private/WorldConditionContext.cpp
--- a/Engine/Plugins/Runtime/WorldConditions/Source/WorldConditions/Private/WorldConditionContext.cpp
+++ b/Engine/Plugins/Runtime/WorldConditions/Source/WorldConditions/Private/WorldConditionContext.cpp
@@ -1,7 +1,7 @@
 // Copyright Epic Games, Inc. All Rights Reserved.
 
 #include "WorldConditionContext.h"
-
+#include "VisualLogger/VisualLogger.h"
 
 bool FWorldConditionContext::Activate() const
 {
@@ -31,10 +31,18 @@ bool FWorldConditionContext::Activate() const
 	for (int32 Index = 0; Index < Conditions.Num(); Index++)
 	{
 		const FWorldConditionBase& Condition = Conditions[Index].Get<const FWorldConditionBase>();
-		bSuccess &= Condition.Activate(*this);
+		if (!Condition.Activate(*this))
+		{
+			UE_VLOG_ALWAYS_UELOG(QueryState.GetOwner(), LogWorldCondition, Error, TEXT("Failed to activate condition at index %d"), Index);
+			bSuccess = false;
+		}
 	}
 
-	if (!bSuccess)
+	if (LIKELY(bSuccess))
+	{
+		QueryState.SetConditionsActivated(true);
+	}
+	else
 	{
 		Deactivate();
 	}
@@ -61,6 +69,11 @@ bool FWorldConditionContext::IsTrue() const
 		return true;
 	}
 
+	if (!ensureMsgf(QueryState.AreConditionsActivated(), TEXT("Can not evaluate non activated conditions. Activate must be called first.")))
+	{
+		return false;
+	}
+
 	static_assert(UE::WorldCondition::MaxExpressionDepth == 4);
 	EWorldConditionResultValue Results[UE::WorldCondition::MaxExpressionDepth + 1] = { EWorldConditionResultValue::Invalid, EWorldConditionResultValue::Invalid, EWorldConditionResultValue::Invalid, EWorldConditionResultValue::Invalid, EWorldConditionResultValue::Invalid };
 	EWorldConditionOperator Operators[UE::WorldCondition::MaxExpressionDepth + 1] = { EWorldConditionOperator::Copy, EWorldConditionOperator::Copy, EWorldConditionOperator::Copy, EWorldConditionOperator::Copy, EWorldConditionOperator::Copy };
@@ -128,5 +141,6 @@ void FWorldConditionContext::Deactivate() const
 		}
 	}
 
+	QueryState.SetConditionsActivated(false);
 	QueryState.Free();
 }
diff --git a/Engine/Plugins/Runtime/WorldConditions/Source/WorldConditions/Private/WorldConditionQuery.cpp b/Engine/Plugins/Runtime/WorldConditions/Source/WorldConditions/Private/WorldConditionQuery.cpp
--- a/Engine/Plugins/Runtime/WorldConditions/Source/WorldConditions/Private/WorldConditionQuery.cpp
+++ b/Engine/Plugins/Runtime/WorldConditions/Source/WorldConditions/Private/WorldConditionQuery.cpp
@@ -65,7 +65,7 @@ FWorldConditionQueryState::~FWorldConditionQueryState()
 	// Try to clean up the best we can in case things are still running.
 	if (IsInitialized() && SharedDefinition.IsValid())
 	{
-		if (IsValid(Owner))
+		if (IsValid(Owner) && AreConditionsActivated())
 		{
 			// We can call Deactivate(), but we dont know the context data, so some uninitialization can potentially be incomplete.
 			UE_LOG(LogWorldCondition, Warning, TEXT("World Condition: State %p owned by %s is still active on destructor, calling Deactivate() without context data, might leak memory or resources."), this, *GetFullNameSafe(Owner));
@@ -75,7 +75,7 @@ FWorldConditionQueryState::~FWorldConditionQueryState()
 		}
 		else
 		{
-			// The owner is not valid, os we cannot call Deactivate(), the best we can do is to clean up the memory properly, as the shared definition is valid. 
+			// The owner is not valid, so we cannot call Deactivate(), the best we can do is to clean up the memory properly, as the shared definition is valid. 
 			UE_LOG(LogWorldCondition, Warning, TEXT("World Condition: State %p owned by %s is still active on destructor, failed to call Deactivate() due to invalid owner, calling Free(), might leak memory or resources."), this, *GetFullNameSafe(Owner));
 			Free();
 		}
diff --git a/Engine/Plugins/Runtime/WorldConditions/Source/WorldConditions/Public/WorldConditionContext.h b/Engine/Plugins/Runtime/WorldConditions/Source/WorldConditions/Public/WorldConditionContext.h
--- a/Engine/Plugins/Runtime/WorldConditions/Source/WorldConditions/Public/WorldConditionContext.h
+++ b/Engine/Plugins/Runtime/WorldConditions/Source/WorldConditions/Public/WorldConditionContext.h
@@ -228,7 +228,7 @@ struct WORLDCONDITIONS_API FWorldConditionContext
 		return *CastChecked<typename T::UStateType>(QueryState.GetStateObject(Condition));
 	}
 
-	/** @return Reference to the query state of update wolrd condition query. */
+	/** @return Reference to the query state of update world condition query. */
 	FWorldConditionQueryState& GetQueryState() const { return QueryState; }
 
 	/**
diff --git a/Engine/Plugins/Runtime/WorldConditions/Source/WorldConditions/Public/WorldConditionQuery.h b/Engine/Plugins/Runtime/WorldConditions/Source/WorldConditions/Public/WorldConditionQuery.h
--- a/Engine/Plugins/Runtime/WorldConditions/Source/WorldConditions/Public/WorldConditionQuery.h
+++ b/Engine/Plugins/Runtime/WorldConditions/Source/WorldConditions/Public/WorldConditionQuery.h
@@ -212,7 +212,7 @@ struct WORLDCONDITIONS_API FWorldConditionQueryDefinition
 	FText GetDescription() const;
 #endif
 
-	const TSharedPtr<FWorldConditionQuerySharedDefinition> GetSharedDefinitionPtr() const { return SharedDefinition; }
+	TSharedPtr<FWorldConditionQuerySharedDefinition> GetSharedDefinitionPtr() const { return SharedDefinition; }
 	
 private:
 	/**
@@ -339,6 +339,7 @@ struct WORLDCONDITIONS_API FWorldConditionQueryState
 
 	FWorldConditionQueryState()
 		: bIsInitialized(false)
+		, bAreConditionsActivated(false)
 	{
 	}
 
@@ -372,6 +373,18 @@ struct WORLDCONDITIONS_API FWorldConditionQueryState
 	 */
 	void Free();
 
+	/** @return True if the conditions were activated. */
+	bool AreConditionsActivated() const
+	{
+		return bAreConditionsActivated;
+	}
+
+	/** Sets the status of the conditions for this state. */
+	void SetConditionsActivated(const bool bConditionsActivated)
+	{
+		bAreConditionsActivated = bConditionsActivated;
+	}
+
 	/** @return cached result stored in the state. */
 	EWorldConditionResultValue GetCachedResult() const
 	{
@@ -454,6 +467,7 @@ private:
 
 	uint8 NumConditions = 0;
 	uint8 bIsInitialized : 1;
+	uint8 bAreConditionsActivated : 1;
 	TSharedPtr<uint8> Memory = nullptr;
 	TSharedPtr<FWorldConditionQuerySharedDefinition> SharedDefinition = nullptr;
 
@@ -484,6 +498,7 @@ struct WORLDCONDITIONS_API FWorldConditionQuery
 
 	/**
 	 * Activates the world conditions in the query.
+	 * @param Owner Owner of any objects created during Activate().
 	 * @param ContextData ContextData that matches the schema of the query.
 	 * @return true of the activation succeeded, or false if failed. Failed queries will return false when IsTrue() is called.
 	 */
