[WorldCondition] updated type validation when setting up context values to use polymorphism for UObject pointers

diff --git a/Engine/Plugins/Runtime/WorldConditions/Source/WorldConditions/Public/WorldConditionContext.h b/Engine/Plugins/Runtime/WorldConditions/Source/WorldConditions/Public/WorldConditionContext.h
--- a/Engine/Plugins/Runtime/WorldConditions/Source/WorldConditions/Public/WorldConditionContext.h
+++ b/Engine/Plugins/Runtime/WorldConditions/Source/WorldConditions/Public/WorldConditionContext.h
@@ -89,7 +89,7 @@ struct WORLDCONDITIONS_API FWorldConditionContextData
 		check(Schema);
 		if (Ref.IsValid())
 		{
-			checkSlow(StructView.GetScriptStruct()->IsChildOf(Schema->GetContextDataDescByRef(Ref)->Struct));
+			check(StructView.GetScriptStruct()->IsChildOf(Schema->GetContextDataDescByRef(Ref)->Struct));
 			Views[Ref.GetIndex()] = FWorldConditionDataView(StructView, Schema->GetContextDataTypeByRef(Ref));
 			return true;
 		}
@@ -103,7 +103,7 @@ struct WORLDCONDITIONS_API FWorldConditionContextData
 		check(Schema);
 		if (Ref.IsValid())
 		{
-			checkSlow(T::StaticStruct()->IsChildOf(Schema->GetContextDataDescByRef(Ref)->Struct));
+			check(T::StaticStruct()->IsChildOf(Schema->GetContextDataDescByRef(Ref)->Struct));
 			Views[Ref.GetIndex()] = FWorldConditionDataView(T::StaticStruct(), reinterpret_cast<const uint8*>(Value), Schema->GetContextDataTypeByRef(Ref));
 			return true;
 		}
@@ -117,7 +117,7 @@ struct WORLDCONDITIONS_API FWorldConditionContextData
 		check(Schema);
 		if (Ref.IsValid())
 		{
-			checkSlow(T::StaticClass()->IsChildOf(Schema->GetContextDataDescByRef(Ref)->Struct));
+			check(Object == nullptr || Object->GetClass()->IsChildOf(Schema->GetContextDataDescByRef(Ref)->Struct));
 			Views[Ref.GetIndex()] = FWorldConditionDataView(Object, Schema->GetContextDataTypeByRef(Ref));
 			return true;
 		}
@@ -132,7 +132,7 @@ struct WORLDCONDITIONS_API FWorldConditionContextData
 		const int32 Index = Schema->GetContextDataIndexByName(Name, T::StaticStruct());
 		if (Index != INDEX_NONE)
 		{
-			checkSlow(T::StaticStruct()->IsChildOf(Schema->GetContextDataDescByIndex(Index).Struct));
+			check(T::StaticStruct()->IsChildOf(Schema->GetContextDataDescByIndex(Index).Struct));
 			Views[Index] = FWorldConditionDataView(T::StaticStruct(), reinterpret_cast<const uint8*>(Value), Schema->GetContextDataTypeByIndex(Index));
 			return true;
 		}
@@ -147,7 +147,7 @@ struct WORLDCONDITIONS_API FWorldConditionContextData
 		const int32 Index = Schema->GetContextDataIndexByName(Name, T::StaticClass());
 		if (Index != INDEX_NONE)
 		{
-			checkSlow(T::StaticClass()->IsChildOf(Schema->GetContextDataDescByIndex(Index).Struct));
+			check(Object == nullptr || Object->GetClass()->IsChildOf(Schema->GetContextDataDescByIndex(Index).Struct));
 			Views[Index] = FWorldConditionDataView(Object, Schema->GetContextDataTypeByIndex(Index));
 			return true;
 		}
