Add support for multi-types world conditions data refs

diff --git a/Engine/Plugins/Runtime/WorldConditions/Source/WorldConditions/Private/WorldConditionSchema.cpp b/Engine/Plugins/Runtime/WorldConditions/Source/WorldConditions/Private/WorldConditionSchema.cpp
--- a/Engine/Plugins/Runtime/WorldConditions/Source/WorldConditions/Private/WorldConditionSchema.cpp
+++ b/Engine/Plugins/Runtime/WorldConditions/Source/WorldConditions/Private/WorldConditionSchema.cpp
@@ -1,6 +1,7 @@
 // Copyright Epic Games, Inc. All Rights Reserved.
 
 #include "WorldConditionSchema.h"
+#include "Algo/AnyOf.h"
 
 #include UE_INLINE_GENERATED_CPP_BY_NAME(WorldConditionSchema)
 
@@ -44,6 +45,17 @@ int32 UWorldConditionSchema::GetContextDataIndexByName(const FName DataName, con
 	});
 }
 
+const FWorldConditionContextDataDesc* UWorldConditionSchema::GetContextDataDescByName(const FName DataName, const TArray<const UStruct*>& Structs) const
+{
+	return ContextDataDescs.FindByPredicate([DataName, &Structs](const FWorldConditionContextDataDesc& Desc)
+	{
+		return Desc.Name == DataName && Algo::AnyOf(Structs, [&Desc] (const UStruct* Struct)
+			{
+				return Desc.Struct->IsChildOf(Struct);
+			});
+	});
+}
+
 const FWorldConditionContextDataDesc* UWorldConditionSchema::GetContextDataDescByName(const FName DataName, const UStruct* Struct) const
 {
 	return ContextDataDescs.FindByPredicate([DataName, Struct](const FWorldConditionContextDataDesc& Desc)
diff --git a/Engine/Plugins/Runtime/WorldConditions/Source/WorldConditions/Public/WorldConditionContext.h b/Engine/Plugins/Runtime/WorldConditions/Source/WorldConditions/Public/WorldConditionContext.h
--- a/Engine/Plugins/Runtime/WorldConditions/Source/WorldConditions/Public/WorldConditionContext.h
+++ b/Engine/Plugins/Runtime/WorldConditions/Source/WorldConditions/Public/WorldConditionContext.h
@@ -109,8 +109,8 @@ struct WORLDCONDITIONS_API FWorldConditionContextData
 		check(Schema);
 		if (Ref.IsValid())
 		{
-			check(T::StaticStruct()->IsChildOf(Schema->GetContextDataDescByRef(Ref)->Struct));
-			Views[Ref.GetIndex()] = FWorldConditionDataView(T::StaticStruct(), reinterpret_cast<const uint8*>(Value), Schema->GetContextDataTypeByRef(Ref));
+			check(TBaseStructure<T>::Get()->IsChildOf(Schema->GetContextDataDescByRef(Ref)->Struct));
+			Views[Ref.GetIndex()] = FWorldConditionDataView(TBaseStructure<T>::Get(), reinterpret_cast<const uint8*>(Value), Schema->GetContextDataTypeByRef(Ref));
 			return true;
 		}
 		return false;
@@ -135,11 +135,11 @@ struct WORLDCONDITIONS_API FWorldConditionContextData
 	typename TEnableIf<!TIsDerivedFrom<T, UObject>::IsDerived, bool>::Type SetContextData(const FName Name, const T* Value)
 	{
 		check(Schema);
-		const int32 Index = Schema->GetContextDataIndexByName(Name, T::StaticStruct());
+		const int32 Index = Schema->GetContextDataIndexByName(Name, TBaseStructure<T>::Get());
 		if (Index != INDEX_NONE)
 		{
-			check(T::StaticStruct()->IsChildOf(Schema->GetContextDataDescByIndex(Index).Struct));
-			Views[Index] = FWorldConditionDataView(T::StaticStruct(), reinterpret_cast<const uint8*>(Value), Schema->GetContextDataTypeByIndex(Index));
+			check(TBaseStructure<T>::Get()->IsChildOf(Schema->GetContextDataDescByIndex(Index).Struct));
+			Views[Index] = FWorldConditionDataView(TBaseStructure<T>::Get(), reinterpret_cast<const uint8*>(Value), Schema->GetContextDataTypeByIndex(Index));
 			return true;
 		}
 		return false;
@@ -175,6 +175,14 @@ struct WORLDCONDITIONS_API FWorldConditionContextData
 		return Views[Ref.GetIndex()].template GetPtr<T>();
 	}
 
+	/** @return Pointer to referenced context data if the type is correct, else returns nullptr. */
+	template <typename T>
+	T* TryGetContextDataPtr(const FWorldConditionContextDataRef& Ref) const
+	{
+		check(Ref.IsValid());
+		return Views[Ref.GetIndex()].template TryGetPtr<T>();
+	}
+
 protected:
 	/** Pointer to schema used to initialize the context data. */
 	const UWorldConditionSchema* Schema = nullptr;
@@ -224,6 +232,16 @@ struct WORLDCONDITIONS_API FWorldConditionContext
 		return ContextData.template GetContextDataPtr<T>(Ref);
 	}
 
+	/**
+	 * Usage Example : auto [Actor, Position] = Context.GetContextDataTuplePtr<AActor, FVector>(Ref);
+	 * @return a tuple of pointers of the given types to the property if possible, nullptr otherwise. 
+	 */
+	template <typename ...T>
+	TTuple<T*...> GetContextDataTuplePtr(const FWorldConditionContextDataRef& Ref) const
+	{
+		return TTuple<T*...>(ContextData.template TryGetContextDataPtr<T>(Ref)...);
+	}
+
 	/** @return Struct State data of the specific world condition. */
 	template <typename T>
 	typename T::FStateType& GetState(const T& Condition) const
diff --git a/Engine/Plugins/Runtime/WorldConditions/Source/WorldConditions/Public/WorldConditionSchema.h b/Engine/Plugins/Runtime/WorldConditions/Source/WorldConditions/Public/WorldConditionSchema.h
--- a/Engine/Plugins/Runtime/WorldConditions/Source/WorldConditions/Public/WorldConditionSchema.h
+++ b/Engine/Plugins/Runtime/WorldConditions/Source/WorldConditions/Public/WorldConditionSchema.h
@@ -52,8 +52,12 @@ public:
 	TConstArrayView<FWorldConditionContextDataDesc> GetContextDataDescs() const;
 
 	/** @return Context data descriptor of specific name and type. */
+	UE_DEPRECATED(5.6, "This method will be deleted, use GetContextDataDescByName with multiple UStruct* as parameters instead to handle multiple typed Context Data.")
 	const FWorldConditionContextDataDesc* GetContextDataDescByName(const FName DataName, const UStruct* Struct) const;
 
+	/** @return Context data descriptor of specific name and types. */
+	const FWorldConditionContextDataDesc* GetContextDataDescByName(const FName DataName, const TArray<const UStruct*>& Structs) const;
+
 	/** @return Context data descriptor of specific name and type. */
 	const FWorldConditionContextDataDesc* GetContextDataDescByRef(const FWorldConditionContextDataRef& Ref) const;
 
@@ -73,32 +77,72 @@ public:
 	int32 GetContextDataIndexByName(const FName DataName, const UStruct* Struct) const;
 
 	/**
-	 * Resolves the index of the context data reference to an Object context data, based on the name in the reference and the type from the template. 
+	 * Resolves the index of the context data reference to an object data, based on the name in the reference and the types from the template. 
+	 * Usage example : Schema.ResolveContextDataRef<AActor, FVector>(SourceWorldConditionContextDataRef)
+	 * @param Ref Reference to resolve.
+	 * @return True if the reference was resolved.
+	 */
+	template <typename Head, typename NextHead, typename... Tail>
+	typename TEnableIf<TIsDerivedFrom<Head, UObject>::IsDerived, bool>::Type ResolveContextDataRef(FWorldConditionContextDataRef& Ref) const
+	{
+		const int32 Index = GetContextDataIndexByName(Ref.Name, Head::StaticClass());
+		if (Index == INDEX_NONE)
+		{
+			return ResolveContextDataRef<NextHead, Tail...>(Ref);
+		}
+
+		Ref.Index = IntCastChecked<uint8>(Index);
+		return true;
+	}
+
+	/**
+	 * Resolves the index of the context data reference to an object data, based on the name in the reference and the type from the template. 
+	 * Usage example : Schema.ResolveContextDataRef<AActor>(SourceActorWorldConditionContextDataRef)
 	 * @param Ref Reference to resolve.
 	 * @return True if the reference was resolved.
 	 */
-	template <typename T>
-	typename TEnableIf<TIsDerivedFrom<T, UObject>::IsDerived, bool>::Type ResolveContextDataRef(FWorldConditionContextDataRef& Ref) const
+	template <typename Head>
+	typename TEnableIf<TIsDerivedFrom<Head, UObject>::IsDerived, bool>::Type ResolveContextDataRef(FWorldConditionContextDataRef& Ref) const
 	{
-		const int32 Index = GetContextDataIndexByName(Ref.Name, T::StaticClass());
+		const int32 Index = GetContextDataIndexByName(Ref.Name, Head::StaticClass());
 		if (Index == INDEX_NONE)
 		{
 			Ref.Index = FWorldConditionContextDataRef::InvalidIndex;
 			return false;
 		}
+
+		Ref.Index = IntCastChecked<uint8>(Index);
+		return true;
+	}
+
+	/**
+	 * Resolves the index of the context data reference to a Struct context data, based on the name in the reference and the types from the template. 
+	 * Usage example : Schema.ResolveContextDataRef<AActor, FVector>(SourceWorldConditionContextDataRef)
+	 * @param Ref Reference to resolve.
+	 * @return True if the reference was resolved.
+	 */
+	template <typename Head, typename NextHead, typename... Tail>
+	typename TEnableIf<!TIsDerivedFrom<Head, UObject>::IsDerived, bool>::Type ResolveContextDataRef(FWorldConditionContextDataRef& Ref) const
+	{
+		const int32 Index = GetContextDataIndexByName(Ref.Name, TBaseStructure<Head>::Get());
+		if (Index == INDEX_NONE)
+		{
+			return ResolveContextDataRef<NextHead, Tail...>(Ref);
+		}
 		Ref.Index = IntCastChecked<uint8>(Index);
 		return true;
 	}
 
 	/**
 	 * Resolves the index of the context data reference to a Struct context data, based on the name in the reference and the type from the template. 
+	 * Usage example : Schema.ResolveContextDataRef<AActor>(SourceActorWorldConditionContextDataRef)
 	 * @param Ref Reference to resolve.
 	 * @return True if the reference was resolved.
 	 */
-	template <typename T>
-	typename TEnableIf<!TIsDerivedFrom<T, UObject>::IsDerived, bool>::Type ResolveContextDataRef(FWorldConditionContextDataRef& Ref) const
+	template <typename Head>
+	typename TEnableIf<!TIsDerivedFrom<Head, UObject>::IsDerived, bool>::Type ResolveContextDataRef(FWorldConditionContextDataRef& Ref) const
 	{
-		const int32 Index = GetContextDataIndexByName(Ref.Name, T::StaticStruct());
+		const int32 Index = GetContextDataIndexByName(Ref.Name, TBaseStructure<Head>::Get());
 		if (Index == INDEX_NONE)
 		{
 			Ref.Index = FWorldConditionContextDataRef::InvalidIndex;
diff --git a/Engine/Plugins/Runtime/WorldConditions/Source/WorldConditions/Public/WorldConditionTypes.h b/Engine/Plugins/Runtime/WorldConditions/Source/WorldConditions/Public/WorldConditionTypes.h
--- a/Engine/Plugins/Runtime/WorldConditions/Source/WorldConditions/Public/WorldConditionTypes.h
+++ b/Engine/Plugins/Runtime/WorldConditions/Source/WorldConditions/Public/WorldConditionTypes.h
@@ -241,10 +241,38 @@ struct FWorldConditionDataView
 	{
 		// If Memory is set, expect Struct too. Otherwise, let nulls pass through.
 		check(!Memory || (Memory && Struct));
-		check(!Struct || Struct->IsChildOf(T::StaticStruct()));
+		check(!Struct || Struct->IsChildOf(TBaseStructure<T>::Get()));
 		return ((T*)Memory);
 	}
 
+	/** @return Pointer to object stored in the view if the type matches correctly, else returns nullptr. */
+	template <typename T>
+	typename TEnableIf<TIsDerivedFrom<T, UObject>::IsDerived, T*>::Type TryGetPtr() const
+	{
+		// If Memory is set, expect Struct too. Otherwise, let nulls pass through.
+		check(!Memory || (Memory && Struct));
+		if (Struct && Struct->IsChildOf(T::StaticClass()))
+		{
+			return ((T*)Memory);
+		}
+
+		return nullptr;
+	}
+
+	/** @return Pointer to struct stored in the view if the type matches correctly, else returns nullptr. */
+	template <typename T>
+	typename TEnableIf<!TIsDerivedFrom<T, UObject>::IsDerived, T*>::Type TryGetPtr() const
+	{
+		// If Memory is set, expect Struct too. Otherwise, let nulls pass through.
+		check(!Memory || (Memory && Struct));
+		if (Struct && Struct->IsChildOf(TBaseStructure<T>::Get()))
+		{
+			return ((T*)Memory);
+		}
+
+		return nullptr;
+	}
+
 	/** @return Struct describing the data type. */
 	const UStruct* GetStruct() const { return Struct; }
 
diff --git a/Engine/Plugins/Runtime/WorldConditions/Source/WorldConditionsEditor/Private/Customizations/WorldConditionContextDataRefDetails.cpp b/Engine/Plugins/Runtime/WorldConditions/Source/WorldConditionsEditor/Private/Customizations/WorldConditionContextDataRefDetails.cpp
--- a/Engine/Plugins/Runtime/WorldConditions/Source/WorldConditionsEditor/Private/Customizations/WorldConditionContextDataRefDetails.cpp
+++ b/Engine/Plugins/Runtime/WorldConditions/Source/WorldConditionsEditor/Private/Customizations/WorldConditionContextDataRefDetails.cpp
@@ -54,27 +54,45 @@ void FWorldConditionContextDataRefDetails::CacheContextData()
 	static const FName BaseStructMetaName(TEXT("BaseStruct")); // TODO: move these names into one central place.
 	static const FName BaseClassMetaName(TEXT("BaseClass")); // TODO: move these names into one central place.
 
-	const FString BaseStructName = StructProperty->GetMetaData(BaseStructMetaName);
-	const FString BaseClassName = StructProperty->GetMetaData(BaseClassMetaName);
+	const FString BaseStructsName = StructProperty->GetMetaData(BaseStructMetaName);
+	const FString BaseClassesName = StructProperty->GetMetaData(BaseClassMetaName);
+
+	TArray<FString> ParsedStructsNames;
+	TArray<FString> ParsedClassesNames;
 	
-	if (!BaseStructName.IsEmpty())
+	if (!BaseStructsName.IsEmpty())
+	{
+		BaseStructsName.ParseIntoArray(ParsedStructsNames, TEXT(","), true);
+	}
+
+	if (!BaseClassesName.IsEmpty())
 	{
-		BaseStruct = UClass::TryFindTypeSlow<UScriptStruct>(BaseStructName);
+		BaseClassesName.ParseIntoArray(ParsedClassesNames, TEXT(","), true);
 	}
 
-	if (BaseStruct == nullptr)
+	for (FString& StructNameString : ParsedStructsNames)
 	{
-		if (!BaseClassName.IsEmpty())
+		StructNameString.TrimStartAndEndInline();
+		if (UStruct* StructType = UClass::TryFindTypeSlow<UScriptStruct>(StructNameString))
 		{
-			BaseStruct = UClass::TryFindTypeSlow<UClass>(BaseClassName);
+			BaseStructs.Emplace(StructType);
 		}
 	}
 
-	if (!BaseStruct)
+	for (FString& ClassNameString : ParsedClassesNames)
+	{
+		ClassNameString.TrimStartAndEndInline();
+		if (UStruct* StructType = UClass::TryFindTypeSlow<UClass>(ClassNameString))
+		{
+			BaseStructs.Emplace(StructType);
+		}
+	}
+
+	if (BaseStructs.IsEmpty())
 	{
 		check(StructProperty.IsValid() && StructProperty->IsValidHandle())
-		UE_LOG(LogWorldCondition, Error, TEXT("%s: Could not find BaseStruct '%s' nor BaseClass '%s' based on the property metadata, expecting full struct name."),
-			*FString(StructProperty->GetPropertyPath()), *BaseStructName, *BaseClassName);
+		UE_LOG(LogWorldCondition, Error, TEXT("%s: Could not find BaseStructs '%s' nor BaseClasses '%s' based on the property metadata, expecting full struct name separated by commas if multiple."),
+			*FString(StructProperty->GetPropertyPath()), *BaseStructsName, *BaseClassesName);
 	}
 
 	// Find schema from outer FWorldConditionQueryDefinition.
@@ -108,14 +126,17 @@ void FWorldConditionContextDataRefDetails::CacheContextData()
 		CurrentProperty = CurrentProperty->GetParentHandle();
 	}
 
-	if (Schema && BaseStruct)
+	if (Schema && !BaseStructs.IsEmpty())
 	{
 		CachedContextData.Reset();
 		for (const FWorldConditionContextDataDesc& Desc : Schema->GetContextDataDescs())
 		{
-			if (Desc.Struct->IsChildOf(BaseStruct))
+			for (const UStruct* BaseStruct : BaseStructs)
 			{
-				CachedContextData.Add(Desc.Name);
+				if (Desc.Struct->IsChildOf(BaseStruct))
+				{
+					CachedContextData.AddUnique(Desc.Name);
+				}
 			}
 		}
 	}
@@ -151,7 +172,7 @@ FText FWorldConditionContextDataRefDetails::GetCurrentDesc() const
 	
 	if (const FWorldConditionContextDataRef* DataRef = GetCommonContextDataRef())
 	{
-		if (const FWorldConditionContextDataDesc* Desc = Schema->GetContextDataDescByName(DataRef->Name, BaseStruct))
+		if (const FWorldConditionContextDataDesc* Desc = Schema->GetContextDataDescByName(DataRef->Name, BaseStructs))
 		{
 			return FText::FromName(Desc->Name);
 		}
diff --git a/Engine/Plugins/Runtime/WorldConditions/Source/WorldConditionsEditor/Private/Customizations/WorldConditionContextDataRefDetails.h b/Engine/Plugins/Runtime/WorldConditions/Source/WorldConditionsEditor/Private/Customizations/WorldConditionContextDataRefDetails.h
--- a/Engine/Plugins/Runtime/WorldConditions/Source/WorldConditionsEditor/Private/Customizations/WorldConditionContextDataRefDetails.h
+++ b/Engine/Plugins/Runtime/WorldConditions/Source/WorldConditionsEditor/Private/Customizations/WorldConditionContextDataRefDetails.h
@@ -35,7 +35,7 @@ private:
 	TArray<FName> CachedContextData;
 
 	const UWorldConditionSchema* Schema = nullptr;
-	const UStruct* BaseStruct = nullptr;
+	TArray<const UStruct*> BaseStructs;
 	
 	IPropertyUtilities* PropUtils = nullptr;
 	TSharedPtr<IPropertyHandle> StructProperty;
