[WorldCondition] adjust widget minimum desired size in customization to better align buttons with generic value widgets

diff --git a/Engine/Plugins/Runtime/WorldConditions/Source/WorldConditionsEditor/Private/Customizations/WorldConditionQueryDefinitionDetails.cpp b/Engine/Plugins/Runtime/WorldConditions/Source/WorldConditionsEditor/Private/Customizations/WorldConditionQueryDefinitionDetails.cpp
--- a/Engine/Plugins/Runtime/WorldConditions/Source/WorldConditionsEditor/Private/Customizations/WorldConditionQueryDefinitionDetails.cpp
+++ b/Engine/Plugins/Runtime/WorldConditions/Source/WorldConditionsEditor/Private/Customizations/WorldConditionQueryDefinitionDetails.cpp
@@ -29,6 +29,7 @@ void FWorldConditionQueryDefinitionDetails::CustomizeHeader(TSharedRef<IProperty
 			StructPropertyHandle->CreatePropertyNameWidget()
 		]
 		.ValueContent()
+		.MinDesiredWidth(170.f)
 		[
 			EditableConditionsProperty->CreatePropertyValueWidget()
 		];
