Fixed CIS issue 767073 by overriding log verbosity during test since it is expected to fail the activation and that would produce an error in the log

diff --git a/Engine/Plugins/Runtime/WorldConditions/Source/WorldConditionsTestSuite/Private/WorldConditionsTest.cpp b/Engine/Plugins/Runtime/WorldConditions/Source/WorldConditionsTestSuite/Private/WorldConditionsTest.cpp
--- a/Engine/Plugins/Runtime/WorldConditions/Source/WorldConditionsTestSuite/Private/WorldConditionsTest.cpp
+++ b/Engine/Plugins/Runtime/WorldConditions/Source/WorldConditionsTestSuite/Private/WorldConditionsTest.cpp
@@ -3,6 +3,7 @@
 #include "AITestsCommon.h"
 #include "Engine/World.h"
 #include "WorldConditionTestTypes.h"
+#include "Logging/LogScopedVerbosityOverride.h"
 #include "Serialization/MemoryReader.h"
 #include "Serialization/MemoryWriter.h"
 #include "Serialization/ObjectAndNameAsStringProxyArchive.h"
@@ -253,11 +254,14 @@ struct FWorldConditionTest_FailingActivate : FAITestBase
 		FWorldConditionContextData ContextData(*TestSchema);
 		const bool bContextDataSet = ContextData.SetContextData(TestSchema->GetValueRef(), &TestData);
 		AITEST_TRUE("Query context data should be set", bContextDataSet);
-		
-		const bool bActivated = Query.Activate(GetWorld(), ContextData);
-		AITEST_FALSE("Query should not get activated", bActivated);
-		AITEST_FALSE("Query should not be active", Query.IsActive());
-		
+
+		{
+			LOG_SCOPE_VERBOSITY_OVERRIDE(LogWorldCondition, ELogVerbosity::Fatal);
+			const bool bActivated = Query.Activate(GetWorld(), ContextData);
+			AITEST_FALSE("Query should not get activated", bActivated);
+			AITEST_FALSE("Query should not be active", Query.IsActive());
+		}
+
 		const bool bResult = Query.IsTrue(ContextData);
 		AITEST_FALSE("Query result should be false", bResult);
 
