[WorldConditions] moved unittests in the AI subgroup

diff --git a/Engine/Plugins/Runtime/WorldConditions/Source/WorldConditionsTestSuite/Private/WorldConditionsTest.cpp b/Engine/Plugins/Runtime/WorldConditions/Source/WorldConditionsTestSuite/Private/WorldConditionsTest.cpp
--- a/Engine/Plugins/Runtime/WorldConditions/Source/WorldConditionsTestSuite/Private/WorldConditionsTest.cpp
+++ b/Engine/Plugins/Runtime/WorldConditions/Source/WorldConditionsTestSuite/Private/WorldConditionsTest.cpp
@@ -37,7 +37,7 @@ struct FWorldConditionTest_Init : FAITestBase
 		return true;
 	}
 };
-IMPLEMENT_AI_INSTANT_TEST(FWorldConditionTest_Init, "System.WorldConditions.Init");
+IMPLEMENT_AI_INSTANT_TEST(FWorldConditionTest_Init, "System.AI.WorldConditions.Init");
 
 struct FWorldConditionTest_Eval : FAITestBase
 {
@@ -72,7 +72,7 @@ struct FWorldConditionTest_Eval : FAITestBase
 		return true;
 	}
 };
-IMPLEMENT_AI_INSTANT_TEST(FWorldConditionTest_Eval, "System.WorldConditions.Eval");
+IMPLEMENT_AI_INSTANT_TEST(FWorldConditionTest_Eval, "System.AI.WorldConditions.Eval");
 
 struct FWorldConditionTest_EvalInvert : FAITestBase
 {
@@ -107,7 +107,7 @@ struct FWorldConditionTest_EvalInvert : FAITestBase
 		return true;
 	}
 };
-IMPLEMENT_AI_INSTANT_TEST(FWorldConditionTest_EvalInvert, "System.WorldConditions.EvalInvert");
+IMPLEMENT_AI_INSTANT_TEST(FWorldConditionTest_EvalInvert, "System.AI.WorldConditions.EvalInvert");
 
 struct FWorldConditionTest_CachedEval : FAITestBase
 {
@@ -160,7 +160,7 @@ struct FWorldConditionTest_CachedEval : FAITestBase
 		return true;
 	}
 };
-IMPLEMENT_AI_INSTANT_TEST(FWorldConditionTest_CachedEval, "System.WorldConditions.CachedEval");
+IMPLEMENT_AI_INSTANT_TEST(FWorldConditionTest_CachedEval, "System.AI.WorldConditions.CachedEval");
 
 struct FWorldConditionTest_EvalComplex : FAITestBase
 {
@@ -203,7 +203,7 @@ struct FWorldConditionTest_EvalComplex : FAITestBase
 		return true;
 	}
 };
-IMPLEMENT_AI_INSTANT_TEST(FWorldConditionTest_EvalComplex, "System.WorldConditions.EvalComplex");
+IMPLEMENT_AI_INSTANT_TEST(FWorldConditionTest_EvalComplex, "System.AI.WorldConditions.EvalComplex");
 
 struct FWorldConditionTest_Empty : FAITestBase
 {
@@ -233,7 +233,7 @@ struct FWorldConditionTest_Empty : FAITestBase
 		return true;
 	}
 };
-IMPLEMENT_AI_INSTANT_TEST(FWorldConditionTest_Empty, "System.WorldConditions.Empty");
+IMPLEMENT_AI_INSTANT_TEST(FWorldConditionTest_Empty, "System.AI.WorldConditions.Empty");
 
 struct FWorldConditionTest_FailingActivate : FAITestBase
 {
@@ -271,7 +271,7 @@ struct FWorldConditionTest_FailingActivate : FAITestBase
 		return true;
 	}
 };
-IMPLEMENT_AI_INSTANT_TEST(FWorldConditionTest_FailingActivate, "System.WorldConditions.FailingActivate");
+IMPLEMENT_AI_INSTANT_TEST(FWorldConditionTest_FailingActivate, "System.AI.WorldConditions.FailingActivate");
 
 struct FWorldConditionTest_Serialization : FAITestBase
 {
@@ -304,7 +304,7 @@ struct FWorldConditionTest_Serialization : FAITestBase
 		return true;
 	}
 };
-IMPLEMENT_AI_INSTANT_TEST(FWorldConditionTest_Serialization, "System.WorldConditions.Serialization");
+IMPLEMENT_AI_INSTANT_TEST(FWorldConditionTest_Serialization, "System.AI.WorldConditions.Serialization");
 
 struct FWorldConditionTest_Description : FAITestBase
 {
@@ -343,7 +343,7 @@ struct FWorldConditionTest_Description : FAITestBase
 		return true;
 	}
 };
-IMPLEMENT_AI_INSTANT_TEST(FWorldConditionTest_Description, "System.WorldConditions.Description");
+IMPLEMENT_AI_INSTANT_TEST(FWorldConditionTest_Description, "System.AI.WorldConditions.Description");
 
 UE_ENABLE_OPTIMIZATION_SHIP
 
