[WorldConditions] Editable conditions are not available outside editor builds to allow hotfixes. By default the conditions are cleared at cook time and are only reinitialized of they are imported again at runtime by the hotfix manager via ImportTextItem

diff --git a/Engine/Plugins/Runtime/WorldConditions/Source/WorldConditions/Private/WorldConditionQuery.cpp b/Engine/Plugins/Runtime/WorldConditions/Source/WorldConditions/Private/WorldConditionQuery.cpp
--- a/Engine/Plugins/Runtime/WorldConditions/Source/WorldConditions/Private/WorldConditionQuery.cpp
+++ b/Engine/Plugins/Runtime/WorldConditions/Source/WorldConditions/Private/WorldConditionQuery.cpp
@@ -526,11 +526,14 @@ FText FWorldConditionQueryDefinition::GetDescription() const
 }
 #endif // WITH_EDITOR
 
-bool FWorldConditionQueryDefinition::Initialize(const UObject* Outer)
+bool FWorldConditionQueryDefinition::Initialize(const UObject* Outer, const EInitializationMode InitializationMode)
 {
-	bool bResult = true;
+	if (InitializationMode == EInitializationMode::UseExistingSharedDefinitionOnEmptyConditions
+		&& EditableConditions.IsEmpty())
+	{
+		return true;
+	}
 
-#if WITH_EDITORONLY_DATA
 	FConstSharedStruct OldSharedDefinition = SharedDefinition;
 	SharedDefinition.Reset();
 
@@ -602,9 +605,7 @@ bool FWorldConditionQueryDefinition::Initialize(const UObject* Outer)
 	Definition.Set(SchemaClass, ValidConditions);
 	Definition.Link(Outer);
 
-#endif // WITH_EDITORONLY_DATA
-	
-	return bResult;
+	return true;
 }
 
 bool FWorldConditionQueryDefinition::Serialize(FArchive& Ar)
@@ -646,11 +647,11 @@ bool FWorldConditionQueryDefinition::Serialize(FArchive& Ar)
 			RawDefinition->PostSerialize(Ar);
 		}
 	}
-#if WITH_EDITOR
+
+#if WITH_EDITORONLY_DATA
 	if (Ar.IsLoading()
 		|| (Ar.IsSaving() && Ar.IsPersistent()))
 	{
- 
 		// If not initialized yet, but has data, initialize on load in editor.
 		if (!SharedDefinition.IsValid() && EditableConditions.Num() > 0)
 		{
@@ -659,8 +660,16 @@ bool FWorldConditionQueryDefinition::Serialize(FArchive& Ar)
 			Initialize(Outer);
 		}
 	}
-#endif
-	
+
+	// Editable conditions should not exist by default when loading WorldConditions
+	// in cooked builds other than for hotfix purposes.
+	if (Ar.IsSaving()
+		&& Ar.IsCooking())
+	{
+		EditableConditions.Reset();
+	}
+#endif // WITH_EDITORONLY_DATA
+
 	return true;
 }
 
@@ -668,7 +677,19 @@ bool FWorldConditionQueryDefinition::ImportTextItem(const TCHAR*& Buffer, int32
 {
 	if (const TCHAR* Result = TBaseStructure<FWorldConditionQueryDefinition>::Get()->ImportText(Buffer, this, Parent, PortFlags, ErrorText, TEXT("FWorldConditionQueryDefinition"), /*bAllowNativeOverride*/false))
 	{
+#if WITH_EDITORONLY_DATA
 		Initialize(Parent);
+#else
+		// In cooked builds, this code path is allowed to reset the SharedDefinition
+		// if EditableConditions is empty in the data imported as text from a hotfix.
+		// e.g., this is required for the case where the definition was cooked with valid conditions,
+		// then they get all removed by the hotfix.
+		Initialize(Parent, EInitializationMode::ResetSharedDefinitionOnEmptyConditions);
+
+		// Editable conditions are no longer necessary after the initialization (i.e., hotfix applied)
+		EditableConditions.Reset();
+#endif
+
 		Buffer = Result;
 		return true;
 	}
diff --git a/Engine/Plugins/Runtime/WorldConditions/Source/WorldConditions/Public/WorldConditionQuery.h b/Engine/Plugins/Runtime/WorldConditions/Source/WorldConditions/Public/WorldConditionQuery.h
--- a/Engine/Plugins/Runtime/WorldConditions/Source/WorldConditions/Public/WorldConditionQuery.h
+++ b/Engine/Plugins/Runtime/WorldConditions/Source/WorldConditions/Public/WorldConditionQuery.h
@@ -3,8 +3,9 @@
 #pragma once
 
 #include "StructUtils/InstancedStructContainer.h"
-#include "WorldConditionBase.h"
 #include "Templates/SubclassOf.h"
+#include "WorldConditionBase.h"
+#include "WorldConditionSchema.h"
 #include "WorldConditionQuery.generated.h"
 
 #define UE_API WORLDCONDITIONS_API
@@ -32,7 +33,7 @@ class FReferenceCollector;
  */
 
 /**
- * Struct used to store a world condition in editor. Used internally.
+ * Struct used to store a world condition in editor or to apply hotfixes. Used internally.
  * Note that the Operator and ExpressionDepth are stored here separately from the World Condition to make sure they are not reset if the Condition is empty. 
  */
 USTRUCT()
@@ -40,7 +41,6 @@ struct FWorldConditionEditable
 {
 	GENERATED_BODY()
 
-#if WITH_EDITORONLY_DATA
 	FWorldConditionEditable() = default;
 
 	FWorldConditionEditable(const uint8 InExpressionDepth, const EWorldConditionOperator InOperator, const FConstStructView InCondition)
@@ -93,7 +93,6 @@ struct FWorldConditionEditable
 	/** Instance of a world condition. */
 	UPROPERTY(EditAnywhere, Category="Default")
 	FInstancedStruct Condition;
-#endif // WITH_EDITORONLY_DATA
 };
 
 /**
@@ -204,12 +203,43 @@ struct FWorldConditionQueryDefinition
 	/** @return true of the definition has conditions and has been initialized. */
 	UE_API bool IsValid() const;
 
+	enum class EInitializationMode : uint8
+	{
+		/**
+		 * The Initialize function will create a SharedDefinition only if there are Editable conditions, otherwise it will be invalid.
+		 */
+		ResetSharedDefinitionOnEmptyConditions,
+
+		/**
+		 * The Initialize function will preserve the existing SharedDefinition if there are no Editable conditions.
+		 */
+		UseExistingSharedDefinitionOnEmptyConditions,
+
+		/**
+		 * The default initialization mode is target specific:
+		 * 
+		 * - In Editor builds:
+		 *   Editable conditions are serialized and represent the source data used to generate the runtime conditions inside the shared definition.
+		 *   SharedDefinition existence is entirely controlled by those conditions, so we use 'ResetSharedDefinitionOnEmptyConditions'.
+		 *
+		 * - In cooked builds:
+		 *   Editable conditions are not serialized in the cooked data so it is expected that the list of Editable conditions
+		 *   will be empty when external systems are calling Initialize, so we use 'UseExistingSharedDefinitionOnEmptyConditions'.
+		 */
+#if WITH_EDITORONLY_DATA
+		Default = ResetSharedDefinitionOnEmptyConditions
+#else
+		Default = UseExistingSharedDefinitionOnEmptyConditions
+#endif
+	};
+
 	/**
 	 * Initialized the condition from editable data.
 	 * @param Outer Used in logging error messages, optional
+	 * @param InitializationMode Indicates how the method will handle an empty list of Editable conditions
 	 * @return true if initialization succeeded.
 	 */
-	UE_API bool Initialize(const UObject* Outer);
+	UE_API bool Initialize(const UObject* Outer, EInitializationMode InitializationMode = EInitializationMode::Default);
 
 	UE_API bool Serialize(FArchive& Ar);
 	UE_API bool ImportTextItem(const TCHAR*& Buffer, int32 PortFlags, UObject* Parent, FOutputDevice* ErrorText, FArchive* InSerializingArchive = nullptr);
@@ -257,11 +287,12 @@ private:
 	UPROPERTY()
 	TSubclassOf<UWorldConditionSchema> SchemaClass = nullptr;
 
-#if WITH_EDITORONLY_DATA
-	/** Conditions used while editing, converted in to Conditions via Initialize(). */
+	/**
+	 * Conditions used while editing, converted in to Conditions via Initialize().
+	 * Those will be cleared for cooked builds but can still be used by hotfixes.
+	 */
 	UPROPERTY(EditAnywhere, Category="Default", meta=(BaseStruct = "/Script/SmartObjectsModule.WorldCondition"))
 	TArray<FWorldConditionEditable> EditableConditions;
-#endif // WITH_EDITORONLY_DATA
 };
 
 template<>
