[WorldConditionQuery] make sure that EditableConditions are cleared for every saves during cooking

diff --git a/Engine/Plugins/Runtime/WorldConditions/Source/WorldConditions/Private/WorldConditionQuery.cpp b/Engine/Plugins/Runtime/WorldConditions/Source/WorldConditions/Private/WorldConditionQuery.cpp
--- a/Engine/Plugins/Runtime/WorldConditions/Source/WorldConditions/Private/WorldConditionQuery.cpp
+++ b/Engine/Plugins/Runtime/WorldConditions/Source/WorldConditions/Private/WorldConditionQuery.cpp
@@ -663,8 +663,8 @@ bool FWorldConditionQueryDefinition::Serialize(FArchive& Ar)
 
 	// Editable conditions should not exist by default when loading WorldConditions
 	// in cooked builds other than for hotfix purposes.
-	if (Ar.IsSaving()
-		&& Ar.IsCooking())
+	if (Ar.IsLoading()
+		&& IsRunningCookCommandlet())
 	{
 		EditableConditions.Reset();
 	}
