[WorldCondition] fixed regression that was causing invalid shared definition to be generated during cook.
Editable conditions were cleared on load but they were still required for calls to Initialize from some PostLoad calls.
With this change we use the previous implementation that was to clear on save but in order to fulfill the primary contract of Cook-Save (package data should not be modified during the save) we back up and clear the conditions before writing the cooked data and restore them after.

diff --git a/Engine/Plugins/Runtime/WorldConditions/Source/WorldConditions/Private/WorldConditionQuery.cpp b/Engine/Plugins/Runtime/WorldConditions/Source/WorldConditions/Private/WorldConditionQuery.cpp
--- a/Engine/Plugins/Runtime/WorldConditions/Source/WorldConditions/Private/WorldConditionQuery.cpp
+++ b/Engine/Plugins/Runtime/WorldConditions/Source/WorldConditions/Private/WorldConditionQuery.cpp
@@ -1,11 +1,12 @@
 // Copyright Epic Games, Inc. All Rights Reserved.
 
 #include "WorldConditionQuery.h"
+#include "Misc/ScopeExit.h"
+#include "Misc/StringBuilder.h"
 #include "PropertyBindingPath.h"
-#include "WorldConditionContext.h"
-#include "UObject/UObjectThreadContext.h"
 #include "Serialization/CustomVersion.h"
-#include "Misc/StringBuilder.h"
+#include "UObject/UObjectThreadContext.h"
+#include "WorldConditionContext.h"
 
 #include UE_INLINE_GENERATED_CPP_BY_NAME(WorldConditionQuery)
 
@@ -612,6 +613,26 @@ bool FWorldConditionQueryDefinition::Serialize(FArchive& Ar)
 {
 	Ar.UsingCustomVersion(FWorldConditionCustomVersion::GUID);
 
+#if WITH_EDITORONLY_DATA
+	// Editable conditions should not exist by default when loading WorldConditions
+	// in cooked builds other than for hotfix purposes.
+	// In order to fulfill the primary contract of Cook-Save (package data should not be modified during the save)
+	// we back up and clear the conditions before writing the cooked data and restore them after
+	TArray<FWorldConditionEditable> ConditionsToRestore;
+	if (Ar.IsSaving()
+		&& Ar.IsCooking())
+	{
+		ConditionsToRestore = MoveTemp(EditableConditions);
+	}
+	ON_SCOPE_EXIT
+	{
+		if (ConditionsToRestore.Num())
+		{
+			EditableConditions = MoveTemp(ConditionsToRestore);
+		}
+	};
+#endif // WITH_EDITORONLY_DATA
+
 	// Use default serialization for most properties for
 	// archive configurations supported by UStruct::SerializeVersionedTaggedProperties 
 	if (Ar.IsLoading() || Ar.IsSaving() || Ar.IsCountingMemory() || Ar.IsObjectReferenceCollector())
@@ -660,14 +681,6 @@ bool FWorldConditionQueryDefinition::Serialize(FArchive& Ar)
 			Initialize(Outer);
 		}
 	}
-
-	// Editable conditions should not exist by default when loading WorldConditions
-	// in cooked builds other than for hotfix purposes.
-	if (Ar.IsLoading()
-		&& IsRunningCookCommandlet())
-	{
-		EditableConditions.Reset();
-	}
 #endif // WITH_EDITORONLY_DATA
 
 	return true;
