[WorldCondition] improved error logging on link failure. Avoid trying to link after each type of serialization; needed only after loading.

diff --git a/Engine/Plugins/Runtime/WorldConditions/Source/WorldConditions/Private/WorldConditionQuery.cpp b/Engine/Plugins/Runtime/WorldConditions/Source/WorldConditions/Private/WorldConditionQuery.cpp
--- a/Engine/Plugins/Runtime/WorldConditions/Source/WorldConditions/Private/WorldConditionQuery.cpp
+++ b/Engine/Plugins/Runtime/WorldConditions/Source/WorldConditions/Private/WorldConditionQuery.cpp
@@ -1,9 +1,11 @@
 // Copyright Epic Games, Inc. All Rights Reserved.
 
 #include "WorldConditionQuery.h"
+#include "Containers/StringFwd.h"
 #include "Misc/ScopeExit.h"
 #include "Misc/StringBuilder.h"
 #include "PropertyBindingPath.h"
+#include "Serialization/ArchiveSerializedPropertyChain.h"
 #include "Serialization/CustomVersion.h"
 #include "UObject/UObjectThreadContext.h"
 #include "WorldConditionContext.h"
@@ -313,11 +315,36 @@ void FWorldConditionQuerySharedDefinition::PostSerialize(const FArchive& Ar)
 {
 	const FUObjectSerializeContext* LoadContext = FUObjectThreadContext::Get().GetSerializeContext();
 	const UObject* SerializedObject = LoadContext ? LoadContext->SerializedObject : nullptr;
-	
-	if (!Link(SerializedObject))
+
+	// We need to link when loading persistent data and when loading during duplication (DuplicateObject or BP re-instancing)
+	if (Ar.IsLoading() && SerializedObject != nullptr)
 	{
-		UE_LOG(LogWorldCondition, Error, TEXT("World Condition: Failed to link query for %s."),
-			*GetNameSafe(SerializedObject));
+		if (!Link(SerializedObject))
+		{
+			UE_SUPPRESS(LogWorldCondition, Error,
+			{
+				FStringBuilderBase PropertyPath;
+				if (const FArchiveSerializedPropertyChain* PropertyChain = Ar.GetSerializedPropertyChain())
+				{
+					const int32 MaxProperties = PropertyChain->GetNumProperties();
+					for (int32 PropertyNum = 0; PropertyNum < MaxProperties; PropertyNum++)
+					{
+						if (const FProperty* CurrProperty = PropertyChain->GetPropertyFromRoot(PropertyNum))
+						{
+							if (PropertyPath.Len() > 0)
+							{
+								PropertyPath << TEXT('.');
+							}
+							PropertyPath << CurrProperty->GetFName();
+						}
+					}
+				}
+
+				UE_LOG(LogWorldCondition, Error, TEXT("World Condition: Failed to link query for property '%s' in '%s'")
+					, *PropertyPath
+					, *GetPathNameSafe(SerializedObject));
+			});
+		}
 	}
 }
 
@@ -351,6 +378,9 @@ bool FWorldConditionQuerySharedDefinition::Link(const UObject* Outer)
 	const UWorldConditionSchema* Schema = SchemaClass.GetDefaultObject();
 	if (!Schema)
 	{
+		UE_LOG(LogWorldCondition, Error, TEXT("World Condition: invalid schema class %s in a query definition for %s.")
+			, *GetPathNameSafe(*SchemaClass)
+			, *GetPathNameSafe(Outer));
 		return false;
 	}
 
