[WorldConditions] fixed serialization unit test now that FWorldConditionQuerySharedDefinition::PostSerialize requires a valid serialized object to link (CIS Issue 989062)

diff --git a/Engine/Plugins/Runtime/WorldConditions/Source/WorldConditionsTestSuite/Private/WorldConditionsTest.cpp b/Engine/Plugins/Runtime/WorldConditions/Source/WorldConditionsTestSuite/Private/WorldConditionsTest.cpp
--- a/Engine/Plugins/Runtime/WorldConditions/Source/WorldConditionsTestSuite/Private/WorldConditionsTest.cpp
+++ b/Engine/Plugins/Runtime/WorldConditions/Source/WorldConditionsTestSuite/Private/WorldConditionsTest.cpp
@@ -7,6 +7,7 @@
 #include "Serialization/MemoryReader.h"
 #include "Serialization/MemoryWriter.h"
 #include "Serialization/ObjectAndNameAsStringProxyArchive.h"
+#include "UObject/UObjectThreadContext.h"
 
 #define LOCTEXT_NAMESPACE "AITestSuite_WorldConditionsTest"
 
@@ -294,7 +295,9 @@ struct FWorldConditionTest_Serialization : FAITestBase
 
 		// Read back on new object
 		Owner = NewAutoDestroyObject<UWorldConditionOwnerClass>();
-		
+
+		FUObjectSerializeContext* SerializeContext = FUObjectThreadContext::Get().GetSerializeContext();
+		TGuardValue<UObject*> SerializedObjectScope(SerializeContext->SerializedObject, Owner);
 		FMemoryReader Reader(Data);
 		FObjectAndNameAsStringProxyArchive ReaderProxy(Reader, /*bInLoadIfFindFails*/true);
 		Owner->Serialize(ReaderProxy);
