[WorldConditions] used binary serialization when required by the archive. The problematic scenario occurs in a cooked build on actor using the flag `bOptimizeBPComponentData` which uses FBlueprintComponentInstanceDataWriter and BinaryPropertySerialization.

diff --git a/Engine/Plugins/Runtime/WorldConditions/Source/WorldConditions/Private/WorldConditionQuery.cpp b/Engine/Plugins/Runtime/WorldConditions/Source/WorldConditions/Private/WorldConditionQuery.cpp
--- a/Engine/Plugins/Runtime/WorldConditions/Source/WorldConditions/Private/WorldConditionQuery.cpp
+++ b/Engine/Plugins/Runtime/WorldConditions/Source/WorldConditions/Private/WorldConditionQuery.cpp
@@ -667,7 +667,14 @@ bool FWorldConditionQueryDefinition::Serialize(FArchive& Ar)
 	// archive configurations supported by UStruct::SerializeVersionedTaggedProperties 
 	if (Ar.IsLoading() || Ar.IsSaving() || Ar.IsCountingMemory() || Ar.IsObjectReferenceCollector())
 	{
-		StaticStruct()->SerializeTaggedProperties(Ar, (uint8*)this, StaticStruct(), nullptr);
+		if (Ar.WantBinaryPropertySerialization())
+		{
+			StaticStruct()->SerializeBin(Ar, this);
+		}
+		else
+		{
+			StaticStruct()->SerializeTaggedProperties(Ar, (uint8*)this, StaticStruct(), nullptr);
+		}
 	}
 
 	if (Ar.CustomVer(FWorldConditionCustomVersion::GUID) >= FWorldConditionCustomVersion::StructSharedDefinition)
@@ -693,7 +700,16 @@ bool FWorldConditionQueryDefinition::Serialize(FArchive& Ar)
 			&& (Ar.IsLoading() || Ar.IsSaving() || Ar.IsCountingMemory() || Ar.IsObjectReferenceCollector()))
 		{
 			UScriptStruct* Struct = TBaseStructure<FWorldConditionQuerySharedDefinition>::Get();
-			Struct->SerializeTaggedProperties(Ar, (uint8*)RawDefinition, Struct, nullptr);
+
+			if (Ar.WantBinaryPropertySerialization())
+			{
+				Struct->SerializeBin(Ar, RawDefinition);
+			}
+			else
+			{
+				Struct->SerializeTaggedProperties(Ar, (uint8*)RawDefinition, Struct, nullptr);
+			}
+
 			// SerializeTaggedProperties does not call PostSerialize() on the struct it's called (calls in items), call it manually.
 			RawDefinition->PostSerialize(Ar);
 		}
