Add support for interfaces in FWorldConditionContextDataRef binding
#Synced-CL 49068285

diff --git a/Engine/Plugins/Runtime/WorldConditions/Source/WorldConditions/Private/WorldConditionSchema.cpp b/Engine/Plugins/Runtime/WorldConditions/Source/WorldConditions/Private/WorldConditionSchema.cpp
--- a/Engine/Plugins/Runtime/WorldConditions/Source/WorldConditions/Private/WorldConditionSchema.cpp
+++ b/Engine/Plugins/Runtime/WorldConditions/Source/WorldConditions/Private/WorldConditionSchema.cpp
@@ -41,7 +41,7 @@ int32 UWorldConditionSchema::GetContextDataIndexByName(const FName DataName, con
 {
 	return ContextDataDescs.IndexOfByPredicate([DataName, Struct](const FWorldConditionContextDataDesc& Desc)
 	{
-		return Desc.Name == DataName && Desc.Struct->IsChildOf(Struct);
+		return Desc.Name == DataName && Desc.MatchesType(Struct);
 	});
 }
 
@@ -51,19 +51,11 @@ const FWorldConditionContextDataDesc* UWorldConditionSchema::GetContextDataDescB
 	{
 		return Desc.Name == DataName && Algo::AnyOf(Structs, [&Desc] (const UStruct* Struct)
 			{
-				return Desc.Struct->IsChildOf(Struct);
+				return Desc.MatchesType(Struct);
 			});
 	});
 }
 
-const FWorldConditionContextDataDesc* UWorldConditionSchema::GetContextDataDescByName(const FName DataName, const UStruct* Struct) const
-{
-	return ContextDataDescs.FindByPredicate([DataName, Struct](const FWorldConditionContextDataDesc& Desc)
-	{
-		return Desc.Name == DataName && Desc.Struct->IsChildOf(Struct);
-	});
-}
-
 FWorldConditionContextDataRef UWorldConditionSchema::AddContextDataDesc(const FName InName, const UStruct* InStruct, const EWorldConditionContextDataType InType)
 {
 	const int32 Index = ContextDataDescs.Num();
diff --git a/Engine/Plugins/Runtime/WorldConditions/Source/WorldConditions/Public/WorldConditionSchema.h b/Engine/Plugins/Runtime/WorldConditions/Source/WorldConditions/Public/WorldConditionSchema.h
--- a/Engine/Plugins/Runtime/WorldConditions/Source/WorldConditions/Public/WorldConditionSchema.h
+++ b/Engine/Plugins/Runtime/WorldConditions/Source/WorldConditions/Public/WorldConditionSchema.h
@@ -53,10 +53,6 @@ public:
 	/** @return All condition descriptors. */
 	UE_API TConstArrayView<FWorldConditionContextDataDesc> GetContextDataDescs() const;
 
-	/** @return Context data descriptor of specific name and type. */
-	UE_DEPRECATED(5.6, "This method will be deleted, use GetContextDataDescByName with multiple UStruct* as parameters instead to handle multiple typed Context Data.")
-	UE_API const FWorldConditionContextDataDesc* GetContextDataDescByName(const FName DataName, const UStruct* Struct) const;
-
 	/** @return Context data descriptor of specific name and types. */
 	UE_API const FWorldConditionContextDataDesc* GetContextDataDescByName(const FName DataName, const TArray<const UStruct*>& Structs) const;
 
diff --git a/Engine/Plugins/Runtime/WorldConditions/Source/WorldConditions/Public/WorldConditionTypes.h b/Engine/Plugins/Runtime/WorldConditions/Source/WorldConditions/Public/WorldConditionTypes.h
--- a/Engine/Plugins/Runtime/WorldConditions/Source/WorldConditions/Public/WorldConditionTypes.h
+++ b/Engine/Plugins/Runtime/WorldConditions/Source/WorldConditions/Public/WorldConditionTypes.h
@@ -123,6 +123,26 @@ struct FWorldConditionContextDataDesc
 	{
 	}
 
+	bool MatchesType(const UStruct* TypeToMatch) const
+	{
+		if (!Struct)
+		{
+			return false;
+		}
+
+		if (Struct->IsChildOf(TypeToMatch))
+		{
+			return true;
+		}
+
+		if (const UClass* DescClass = Cast<UClass>(Struct))
+		{
+			return DescClass->ImplementsInterface(Cast<UClass>(TypeToMatch));
+		}
+
+		return false;
+	}
+
 	/** Type of the context data, can be a struct or object. */
 	UPROPERTY()
 	TObjectPtr<const UStruct> Struct;
diff --git a/Engine/Plugins/Runtime/WorldConditions/Source/WorldConditionsEditor/Private/Customizations/WorldConditionContextDataRefDetails.cpp b/Engine/Plugins/Runtime/WorldConditions/Source/WorldConditionsEditor/Private/Customizations/WorldConditionContextDataRefDetails.cpp
--- a/Engine/Plugins/Runtime/WorldConditions/Source/WorldConditionsEditor/Private/Customizations/WorldConditionContextDataRefDetails.cpp
+++ b/Engine/Plugins/Runtime/WorldConditions/Source/WorldConditionsEditor/Private/Customizations/WorldConditionContextDataRefDetails.cpp
@@ -133,7 +133,7 @@ void FWorldConditionContextDataRefDetails::CacheContextData()
 		{
 			for (const UStruct* BaseStruct : BaseStructs)
 			{
-				if (Desc.Struct->IsChildOf(BaseStruct))
+				if (Desc.MatchesType(BaseStruct))
 				{
 					CachedContextData.AddUnique(Desc.Name);
 				}
