[WorldConditions][Editor] Validate schema on paste of editable conditions

Data implementers often copy and paste World Conditions between queries with different schemas. Previously, pasting a condition not allowed by the target query's schema would silently succeed, making it difficult to diagnose why conditions were not working at runtime.

Changes:
- Add custom PasteAction to FWorldConditionEditableDetails header row
- Validate pasted condition struct type against Schema::IsStructAllowed before allowing paste
- Disallowed conditions now grey out the Paste option in the context menu

#Synced-CL 50904216

diff --git a/Engine/Plugins/Runtime/WorldConditions/Source/WorldConditionsEditor/Private/Customizations/WorldConditionEditableDetails.cpp b/Engine/Plugins/Runtime/WorldConditions/Source/WorldConditionsEditor/Private/Customizations/WorldConditionEditableDetails.cpp
--- a/Engine/Plugins/Runtime/WorldConditions/Source/WorldConditionsEditor/Private/Customizations/WorldConditionEditableDetails.cpp
+++ b/Engine/Plugins/Runtime/WorldConditions/Source/WorldConditionsEditor/Private/Customizations/WorldConditionEditableDetails.cpp
@@ -16,6 +16,8 @@
 #include "WorldConditionSchema.h"
 #include "Editor.h"
 #include "Modules/ModuleManager.h"
+#include "HAL/PlatformApplicationMisc.h"
+#include "ScopedTransaction.h"
 
 #define LOCTEXT_NAMESPACE "WorldCondition"
 
@@ -81,6 +83,10 @@ void FWorldConditionEditableDetails::CustomizeHeader(TSharedRef<IPropertyHandle>
 	const FResetToDefaultOverride ResetOverride = FResetToDefaultOverride::Create(IsResetVisible, ResetHandler);
 
 	HeaderRow
+		.CopyAction(FUIAction(FExecuteAction::CreateSP(this, &FWorldConditionEditableDetails::OnCopyCondition)))
+		.PasteAction(FUIAction(
+			FExecuteAction::CreateSP(this, &FWorldConditionEditableDetails::OnPasteCondition),
+			FCanExecuteAction::CreateSP(this, &FWorldConditionEditableDetails::CanPasteCondition)))
 		.WholeRowContent()
 		.VAlign(VAlign_Center)
 		[
@@ -651,6 +657,59 @@ void FWorldConditionEditableDetails::OnStructPicked(const UScriptStruct* InStruc
 	}
 }
 
+void FWorldConditionEditableDetails::OnCopyCondition() const
+{
+	FString Value;
+	StructProperty->GetValueAsFormattedString(Value, PPF_Copy);
+	FPlatformApplicationMisc::ClipboardCopy(*Value);
+}
+
+bool FWorldConditionEditableDetails::CanPasteCondition() const
+{
+	// Never allowing pasting with no schema, we are not supporting anything basically.
+	if (!Schema)
+	{
+		return false;
+	}
+
+	FString ClipboardText;
+	FPlatformApplicationMisc::ClipboardPaste(ClipboardText);
+	if (ClipboardText.IsEmpty())
+	{
+		return false;
+	}
+
+	// Try to deserialize the clipboard data as a FWorldConditionEditable.
+	FWorldConditionEditable TempEditable;
+	const TCHAR* Buffer = *ClipboardText;
+	const TCHAR* Result = FWorldConditionEditable::StaticStruct()->ImportText(
+		Buffer, &TempEditable, nullptr, PPF_Copy, nullptr, FWorldConditionEditable::StaticStruct()->GetName());
+	if (!Result)
+	{
+		return false;
+	}
+
+	// If the condition is valid, check if the struct type is allowed.
+	if (TempEditable.Condition.IsValid())
+	{
+		return Schema->IsStructAllowed(TempEditable.Condition.GetScriptStruct());
+	}
+
+	return true;
+}
+
+void FWorldConditionEditableDetails::OnPasteCondition() const
+{
+	FString Value;
+	FPlatformApplicationMisc::ClipboardPaste(Value);
+
+	if (!Value.IsEmpty())
+	{
+		FScopedTransaction Transaction(LOCTEXT("PasteCondition", "Paste Condition"));
+		StructProperty->SetValueFromFormattedString(Value, EPropertyValueSetFlags::InstanceObjects);
+	}
+}
+
 const FWorldConditionEditable* FWorldConditionEditableDetails::GetCommonCondition() const
 {
 	check(StructProperty);
diff --git a/Engine/Plugins/Runtime/WorldConditions/Source/WorldConditionsEditor/Private/Customizations/WorldConditionEditableDetails.h b/Engine/Plugins/Runtime/WorldConditions/Source/WorldConditionsEditor/Private/Customizations/WorldConditionEditableDetails.h
--- a/Engine/Plugins/Runtime/WorldConditions/Source/WorldConditionsEditor/Private/Customizations/WorldConditionEditableDetails.h
+++ b/Engine/Plugins/Runtime/WorldConditions/Source/WorldConditionsEditor/Private/Customizations/WorldConditionEditableDetails.h
@@ -62,6 +62,10 @@ private:
 
 	TSharedRef<SWidget> GeneratePicker();
 	void OnStructPicked(const UScriptStruct* InStruct) const;
+
+	void OnCopyCondition() const;
+	bool CanPasteCondition() const;
+	void OnPasteCondition() const;
 	
 	const FWorldConditionEditable* GetCommonCondition() const;
 
