[WorldConditions] Make WorldCondition queries and query definition GetDescription available outside editor

Removes the WITH_EDITOR guard so that human-readable condition descriptions can be used in non-editor contexts such as debug tools (e.g. ImGui menus) and log output. This also enables downstream world condition classes to remove WITH_EDITOR from their GetDescription() overrides.

Changes:
- Remove WITH_EDITOR guard from FWorldConditionBase::GetDescription() declaration and implementation; use GetDisplayNameText() in editor builds, GetAuthoredName() fallback otherwise
- Remove WITH_EDITOR guard from FWorldConditionQueryDefinition::GetDescription() implementation
- Fix typo in doc comment ([REDACTED_HANDLE] -> [REDACTED_HANDLE])

#Synced-CL 51456608

diff --git a/Engine/Plugins/Runtime/WorldConditions/Source/WorldConditions/Private/WorldConditionBase.cpp b/Engine/Plugins/Runtime/WorldConditions/Source/WorldConditions/Private/WorldConditionBase.cpp
--- a/Engine/Plugins/Runtime/WorldConditions/Source/WorldConditions/Private/WorldConditionBase.cpp
+++ b/Engine/Plugins/Runtime/WorldConditions/Source/WorldConditions/Private/WorldConditionBase.cpp
@@ -9,12 +9,14 @@ FWorldConditionBase::~FWorldConditionBase()
 	// Empty
 }
 
-#if WITH_EDITOR
 FText FWorldConditionBase::GetDescription() const
 {
+#if WITH_EDITORONLY_DATA
 	return StaticStruct()->GetDisplayNameText();
-}
+#else
+	return FText::FromString(StaticStruct()->GetAuthoredName());
 #endif
+}
 
 bool FWorldConditionBase::Initialize(const UWorldConditionSchema& Schema)
 {
diff --git a/Engine/Plugins/Runtime/WorldConditions/Source/WorldConditions/Private/WorldConditionQuery.cpp b/Engine/Plugins/Runtime/WorldConditions/Source/WorldConditions/Private/WorldConditionQuery.cpp
--- a/Engine/Plugins/Runtime/WorldConditions/Source/WorldConditions/Private/WorldConditionQuery.cpp
+++ b/Engine/Plugins/Runtime/WorldConditions/Source/WorldConditions/Private/WorldConditionQuery.cpp
@@ -478,7 +478,6 @@ bool FWorldConditionQueryDefinition::Initialize(const UObject* Outer, const TSub
 }
 #endif
 
-#if WITH_EDITOR
 FText FWorldConditionQueryDefinition::GetDescription() const
 {
 	const FWorldConditionQuerySharedDefinition* Definition = SharedDefinition.GetPtr<FWorldConditionQuerySharedDefinition>();
@@ -555,7 +554,6 @@ FText FWorldConditionQueryDefinition::GetDescription() const
 
 	return FText::FromString(Builder.ToString());
 }
-#endif // WITH_EDITOR
 
 bool FWorldConditionQueryDefinition::Initialize(const UObject* Outer, const EInitializationMode InitializationMode)
 {
diff --git a/Engine/Plugins/Runtime/WorldConditions/Source/WorldConditions/Public/WorldConditionBase.h b/Engine/Plugins/Runtime/WorldConditions/Source/WorldConditions/Public/WorldConditionBase.h
--- a/Engine/Plugins/Runtime/WorldConditions/Source/WorldConditions/Public/WorldConditionBase.h
+++ b/Engine/Plugins/Runtime/WorldConditions/Source/WorldConditions/Public/WorldConditionBase.h
@@ -129,10 +129,8 @@ struct FWorldConditionBase
 	/** @return The Instance data type of the condition. */
 	virtual TObjectPtr<const UStruct>* GetRuntimeStateType() const { return nullptr; }
 
-#if WITH_EDITOR
-	/** @return Description to be shown in the UI. */
+	/** @return Description to be shown in the UI or any debug tool. */
 	UE_API virtual FText GetDescription() const;
-#endif
 
 	/**
 	 * Initializes the condition to be used with a specific schema.
diff --git a/Engine/Plugins/Runtime/WorldConditions/Source/WorldConditions/Public/WorldConditionQuery.h b/Engine/Plugins/Runtime/WorldConditions/Source/WorldConditions/Public/WorldConditionQuery.h
--- a/Engine/Plugins/Runtime/WorldConditions/Source/WorldConditions/Public/WorldConditionQuery.h
+++ b/Engine/Plugins/Runtime/WorldConditions/Source/WorldConditions/Public/WorldConditionQuery.h
@@ -258,10 +258,8 @@ struct FWorldConditionQueryDefinition
 	}
 #endif
 
-#if WITH_EDITOR
-	/** @returs description of the query expression. */
+	/** @returns description of the query expression. */
 	UE_API FText GetDescription() const;
-#endif
 
 	FWorldConditionQuerySharedDefinition* GetSharedDefinition() const
 	{
