[WorldCondition] fixed case where world condition queries would not be properly linked when embedded in a struct

#Synced-CL 51542861

diff --git a/Engine/Plugins/Runtime/WorldConditions/Source/WorldConditions/Private/WorldConditionQuery.cpp b/Engine/Plugins/Runtime/WorldConditions/Source/WorldConditions/Private/WorldConditionQuery.cpp
--- a/Engine/Plugins/Runtime/WorldConditions/Source/WorldConditions/Private/WorldConditionQuery.cpp
+++ b/Engine/Plugins/Runtime/WorldConditions/Source/WorldConditions/Private/WorldConditionQuery.cpp
@@ -313,12 +313,12 @@ void FWorldConditionQuerySharedDefinition::Set(const TSubclassOf<UWorldCondition
 
 void FWorldConditionQuerySharedDefinition::PostSerialize(const FArchive& Ar)
 {
-	const FUObjectSerializeContext* LoadContext = FUObjectThreadContext::Get().GetSerializeContext();
-	const UObject* SerializedObject = LoadContext ? LoadContext->SerializedObject : nullptr;
-
 	// We need to link when loading persistent data and when loading during duplication (DuplicateObject or BP re-instancing)
-	if (Ar.IsLoading() && SerializedObject != nullptr)
+	if (Ar.IsLoading())
 	{
+		const FUObjectSerializeContext* LoadContext = FUObjectThreadContext::Get().GetSerializeContext();
+		const UObject* SerializedObject = LoadContext ? LoadContext->SerializedObject : nullptr;
+
 		if (!Link(SerializedObject))
 		{
 			UE_SUPPRESS(LogWorldCondition, Error,
diff --git a/Engine/Plugins/Runtime/WorldConditions/Source/WorldConditionsTestSuite/Private/WorldConditionsTest.cpp b/Engine/Plugins/Runtime/WorldConditions/Source/WorldConditionsTestSuite/Private/WorldConditionsTest.cpp
--- a/Engine/Plugins/Runtime/WorldConditions/Source/WorldConditionsTestSuite/Private/WorldConditionsTest.cpp
+++ b/Engine/Plugins/Runtime/WorldConditions/Source/WorldConditionsTestSuite/Private/WorldConditionsTest.cpp
@@ -309,6 +309,43 @@ struct FWorldConditionTest_Serialization : FAITestBase
 };
 IMPLEMENT_AI_INSTANT_TEST(FWorldConditionTest_Serialization, "System.AI.WorldConditions.Serialization");
 
+// Tests that FWorldConditionQuerySharedDefinition is properly linked when FWorldConditionQueryDefinition
+// is nested inside a UStruct. In that case, UObject serialize context (SerializedObject) may be null
+// when PostSerialize is called, and Link() must still be invoked.
+struct FWorldConditionTest_SerializationInStruct : FAITestBase
+{
+	virtual bool InstantTest() override
+	{
+		UWorldConditionOwnerClassWithNestedStruct* Owner = NewAutoDestroyObject<UWorldConditionOwnerClassWithNestedStruct>();
+		const bool bInitialized = Owner->Container.Definition.Initialize(&GetWorld(), UWorldConditionTestSchema::StaticClass(),
+			{
+				FWorldConditionEditable(0, EWorldConditionOperator::Copy, FConstStructView::Make(FWorldConditionTest())),
+				FWorldConditionEditable(0, EWorldConditionOperator::And, FConstStructView::Make(FWorldConditionTest())),
+			});
+
+		AITEST_TRUE("Query definition should get initialized", bInitialized);
+		AITEST_TRUE("Query definition should be valid before serialization", Owner->Container.Definition.IsValid());
+
+		TArray<uint8> Data;
+		FMemoryWriter Writer(Data);
+		FObjectAndNameAsStringProxyArchive WriterProxy(Writer, /*bInLoadIfFindFails*/true);
+		Owner->Serialize(WriterProxy);
+
+		// Read back into a new object WITHOUT setting SerializedObject in the context.
+		// This replicates the case where FWorldConditionQueryDefinition is nested inside a UStruct,
+		// where the UObject serialize context is not set up by the caller (e.g., raw struct serialization).
+		Owner = NewAutoDestroyObject<UWorldConditionOwnerClassWithNestedStruct>();
+		FMemoryReader Reader(Data);
+		FObjectAndNameAsStringProxyArchive ReaderProxy(Reader, /*bInLoadIfFindFails*/true);
+		Owner->Serialize(ReaderProxy);
+
+		AITEST_TRUE("Query definition nested in a struct should be linked after loading", Owner->Container.Definition.IsValid());
+
+		return true;
+	}
+};
+IMPLEMENT_AI_INSTANT_TEST(FWorldConditionTest_SerializationInStruct, "System.AI.WorldConditions.SerializationInStruct");
+
 struct FWorldConditionTest_Description : FAITestBase
 {
 	virtual bool InstantTest() override
diff --git a/Engine/Plugins/Runtime/WorldConditions/Source/WorldConditionsTestSuite/Public/WorldConditionTestTypes.h b/Engine/Plugins/Runtime/WorldConditions/Source/WorldConditionsTestSuite/Public/WorldConditionTestTypes.h
--- a/Engine/Plugins/Runtime/WorldConditions/Source/WorldConditionsTestSuite/Public/WorldConditionTestTypes.h
+++ b/Engine/Plugins/Runtime/WorldConditions/Source/WorldConditionsTestSuite/Public/WorldConditionTestTypes.h
@@ -229,7 +229,27 @@ class UWorldConditionOwnerClass : public UObject
 	GENERATED_BODY()
 
 public:
-	
+
+	UPROPERTY(EditAnywhere, Category = "Default")
+	FWorldConditionQueryDefinition Definition;
+};
+
+USTRUCT()
+struct FWorldConditionQueryDefinitionContainer
+{
+	GENERATED_BODY()
+
 	UPROPERTY(EditAnywhere, Category = "Default")
 	FWorldConditionQueryDefinition Definition;
 };
+
+UCLASS(hidden)
+class UWorldConditionOwnerClassWithNestedStruct : public UObject
+{
+	GENERATED_BODY()
+
+public:
+
+	UPROPERTY(EditAnywhere, Category = "Default")
+	FWorldConditionQueryDefinitionContainer Container;
+};
